Graph Visualizer: Replacing Control Rig AST visualizer with common service

diff --git a/Engine/Plugins/Animation/ControlRig/Source/ControlRigDeveloper/ControlRigDeveloper.Build.cs b/Engine/Plugins/Animation/ControlRig/Source/ControlRigDeveloper/ControlRigDeveloper.Build.cs
--- a/Engine/Plugins/Animation/ControlRig/Source/ControlRigDeveloper/ControlRigDeveloper.Build.cs
+++ b/Engine/Plugins/Animation/ControlRig/Source/ControlRigDeveloper/ControlRigDeveloper.Build.cs
@@ -38,6 +38,7 @@ namespace UnrealBuildTool.Rules
                 new string[]
                 {
                     "AnimationCore",
+                    "VisualGraphUtils",
                 }
             );
 
diff --git a/Engine/Plugins/Animation/ControlRig/Source/ControlRigDeveloper/Private/ControlRigVisualGraphUtils.cpp b/Engine/Plugins/Animation/ControlRig/Source/ControlRigDeveloper/Private/ControlRigVisualGraphUtils.cpp
new file mode 100644
--- /dev/null
+++ b/Engine/Plugins/Animation/ControlRig/Source/ControlRigDeveloper/Private/ControlRigVisualGraphUtils.cpp
@@ -0,0 +1,115 @@
+// Copyright Epic Games, Inc. All Rights Reserved.
+
+#include "ControlRigVisualGraphUtils.h"
+
+FString FControlRigVisualGraphUtils::DumpRigHierarchyToDotGraph(URigHierarchy* InHierarchy)
+{
+	check(InHierarchy);
+
+	FVisualGraph Graph(TEXT("Rig"));
+
+	struct Local
+	{
+		static TArray<int32> VisitParents(const FRigBaseElement* InElement, URigHierarchy* InHierarchy, FVisualGraph& OutGraph)
+		{
+			TArray<int32> Result;
+			FRigBaseElementParentArray Parents = InHierarchy->GetParents(InElement);
+
+			for(int32 ParentIndex = 0; ParentIndex < Parents.Num(); ParentIndex++)
+			{
+				const int32 ParentNodeIndex = VisitElement(Parents[ParentIndex], InHierarchy, OutGraph);
+				Result.Add(ParentNodeIndex);
+			}
+
+			return Result;
+		}
+
+		static int32 VisitElement(const FRigBaseElement* InElement, URigHierarchy* InHierarchy, FVisualGraph& OutGraph)
+		{
+			if(InElement->GetType() == ERigElementType::Curve)
+			{
+				return INDEX_NONE;
+			}
+			
+			const FName NodeName = *FString::Printf(TEXT("Element_%d"), InElement->GetIndex());
+			int32 NodeIndex = OutGraph.FindNode(NodeName);
+			if(NodeIndex != INDEX_NONE)
+			{
+				return NodeIndex;
+			}
+
+			EVisualGraphShape Shape = EVisualGraphShape::Ellipse;
+			TOptional<FLinearColor> Color;
+
+			switch(InElement->GetType())
+			{
+				case ERigElementType::Bone:
+				{
+					Shape = EVisualGraphShape::Box;
+						
+					if(Cast<FRigBoneElement>(InElement)->BoneType == ERigBoneType::User)
+					{
+						Color = FLinearColor::Green;
+					}
+					break;
+				}
+				case ERigElementType::Null:
+				{
+					Shape = EVisualGraphShape::Diamond;
+					break;
+				}
+				case ERigElementType::Control:
+				{
+					Color = Cast<FRigControlElement>(InElement)->Settings.GizmoColor;
+					break;
+				}
+				default:
+				{
+					break;
+				}
+			}
+
+			NodeIndex = OutGraph.AddNode(NodeName, InElement->GetName(), Color, Shape);
+
+			if(NodeIndex != INDEX_NONE)
+			{
+				TArray<int32> Parents = VisitParents(InElement, InHierarchy, OutGraph);
+				TArray<FRigElementWeight> Weights = InHierarchy->GetParentWeightArray(InElement);
+				for(int32 ParentIndex = 0; ParentIndex < Parents.Num(); ParentIndex++)
+				{
+					const int32 ParentNodeIndex = Parents[ParentIndex];
+					if(ParentNodeIndex != INDEX_NONE)
+					{
+						const TOptional<FLinearColor> EdgeColor;
+						TOptional<EVisualGraphStyle> Style;
+						if(Weights.IsValidIndex(ParentIndex))
+						{
+							if(Weights[ParentIndex].IsAlmostZero())
+							{
+								Style = EVisualGraphStyle::Dotted;
+							}
+						}
+						OutGraph.AddEdge(
+							NodeIndex,
+							ParentNodeIndex,
+							EVisualGraphEdgeDirection::SourceToTarget,
+							NAME_None,
+							TOptional<FName>(),
+							EdgeColor,
+							Style);
+					}
+				}
+			}
+
+			return NodeIndex;
+		}
+	};
+
+	InHierarchy->ForEach([InHierarchy, &Graph](const FRigBaseElement* InElement)
+	{
+		Local::VisitElement(InElement, InHierarchy, Graph);
+		return true;
+	});
+
+	return Graph.DumpDot();
+}
\ No newline at end of file
diff --git a/Engine/Plugins/Animation/ControlRig/Source/ControlRigDeveloper/Public/ControlRigVisualGraphUtils.h b/Engine/Plugins/Animation/ControlRig/Source/ControlRigDeveloper/Public/ControlRigVisualGraphUtils.h
new file mode 100644
--- /dev/null
+++ b/Engine/Plugins/Animation/ControlRig/Source/ControlRigDeveloper/Public/ControlRigVisualGraphUtils.h
@@ -0,0 +1,12 @@
+// Copyright Epic Games, Inc. All Rights Reserved.
+
+#pragma once
+
+#include "CoreMinimal.h"
+#include "VisualGraphUtils.h"
+#include "Rigs/RigHierarchy.h"
+
+struct CONTROLRIGDEVELOPER_API FControlRigVisualGraphUtils
+{
+	static FString DumpRigHierarchyToDotGraph(URigHierarchy* InHierarchy);
+};
diff --git a/Engine/Plugins/Animation/ControlRig/Source/ControlRigEditor/Private/ControlRigCompilerDetails.cpp b/Engine/Plugins/Animation/ControlRig/Source/ControlRigEditor/Private/ControlRigCompilerDetails.cpp
--- a/Engine/Plugins/Animation/ControlRig/Source/ControlRigEditor/Private/ControlRigCompilerDetails.cpp
+++ b/Engine/Plugins/Animation/ControlRig/Source/ControlRigEditor/Private/ControlRigCompilerDetails.cpp
@@ -13,6 +13,7 @@
 #include "IPropertyUtilities.h"
 #include "IPythonScriptPlugin.h"
 #include "RigVMPythonUtils.h"
+#include "ControlRigVisualGraphUtils.h"
 
 #define LOCTEXT_NAMESPACE "ControlRigCompilerDetails"
 
@@ -92,11 +93,11 @@ void FRigVMCompileSettingsDetails::CustomizeChildren(TSharedRef<IPropertyHandle>
 			];
 #endif
 
-		StructBuilder.AddCustomRow(LOCTEXT("ASTTools", "AST Tools"))
+		StructBuilder.AddCustomRow(LOCTEXT("DebuggingTools", "Debugging Tools"))
 			.NameContent()
 			[
 				SNew(STextBlock)
-				.Text(FText::FromString(TEXT("AST")))
+				.Text(FText::FromString(TEXT("Debugging")))
 			.Font(IDetailLayoutBuilder::GetDetailFont())
 			]
 			.ValueContent()
@@ -111,7 +112,7 @@ void FRigVMCompileSettingsDetails::CustomizeChildren(TSharedRef<IPropertyHandle>
 					[
 						SNew(STextBlock)
 						.Justification(ETextJustify::Center)
-						.Text(LOCTEXT("CopyASTToClipboard", "Copy AST"))
+						.Text(LOCTEXT("CopyASTToClipboard", "Copy AST Graph"))
 					]
 				]
 				+ SVerticalBox::Slot()
@@ -126,6 +127,18 @@ void FRigVMCompileSettingsDetails::CustomizeChildren(TSharedRef<IPropertyHandle>
 						.Text(LOCTEXT("CopyByteCodeToClipboard", "Copy ByteCode"))
 					]
 				]
+				+ SVerticalBox::Slot()
+				[
+					SNew(SButton)
+					.OnClicked(this, &FRigVMCompileSettingsDetails::OnCopyHierarchyGraphClicked)
+					.ContentPadding(FMargin(2))
+					.Content()
+					[
+						SNew(STextBlock)
+						.Justification(ETextJustify::Center)
+						.Text(LOCTEXT("CopyHierarchyGraphToClipboard", "Copy Hierarchy Graph"))
+					]
+				]
 			];
 	}
 }
@@ -179,4 +192,17 @@ FReply FRigVMCompileSettingsDetails::OnCopyByteCodeClicked()
 	return FReply::Handled();
 }
 
+FReply FRigVMCompileSettingsDetails::OnCopyHierarchyGraphClicked()
+{
+	if (BlueprintBeingCustomized)
+	{
+		if(UControlRig* ControlRig = Cast<UControlRig>(BlueprintBeingCustomized->GetObjectBeingDebugged()))
+		{
+			FString DotGraphContent = FControlRigVisualGraphUtils::DumpRigHierarchyToDotGraph(ControlRig->GetHierarchy());
+			FPlatformApplicationMisc::ClipboardCopy(*DotGraphContent);
+		}
+	}
+	return FReply::Handled();
+}
+
 #undef LOCTEXT_NAMESPACE
diff --git a/Engine/Plugins/Animation/ControlRig/Source/ControlRigEditor/Private/ControlRigCompilerDetails.h b/Engine/Plugins/Animation/ControlRig/Source/ControlRigEditor/Private/ControlRigCompilerDetails.h
--- a/Engine/Plugins/Animation/ControlRig/Source/ControlRigEditor/Private/ControlRigCompilerDetails.h
+++ b/Engine/Plugins/Animation/ControlRig/Source/ControlRigEditor/Private/ControlRigCompilerDetails.h
@@ -29,6 +29,7 @@ protected:
 #endif
 	FReply OnCopyASTClicked();
 	FReply OnCopyByteCodeClicked();
+	FReply OnCopyHierarchyGraphClicked();
 	
 	UControlRigBlueprint* BlueprintBeingCustomized;
 };
diff --git a/Engine/Source/Developer/RigVMDeveloper/Private/RigVMCompiler/RigVMAST.cpp b/Engine/Source/Developer/RigVMDeveloper/Private/RigVMCompiler/RigVMAST.cpp
--- a/Engine/Source/Developer/RigVMDeveloper/Private/RigVMCompiler/RigVMAST.cpp
+++ b/Engine/Source/Developer/RigVMDeveloper/Private/RigVMCompiler/RigVMAST.cpp
@@ -19,6 +19,7 @@
 #include "RigVMCore/RigVMExecuteContext.h"
 #include "Stats/StatsHierarchical.h"
 #include "RigVMDeveloperModule.h"
+#include "VisualGraphUtils.h"
 
 FRigVMExprAST::FRigVMExprAST(EType InType, const FRigVMASTProxy& InProxy)
 	: Name(NAME_None)
@@ -322,198 +323,6 @@ FString FRigVMExprAST::DumpText(const FString& InPrefix) const
 	return Result;
 }
 
-FString FRigVMExprAST::DumpDot(TArray<bool>& OutExpressionDefined, const FString& InPrefix) const
-{
-	FString Prefix = InPrefix;
-
-	FString Result;
-	bool bWasDefined = true;
-	if (!OutExpressionDefined[GetIndex()])
-	{
-		bWasDefined = false;
-
-		FString Label = GetName().ToString();
-		FString AdditionalNodeSettings;
-		switch (GetType())
-		{
-			case EType::Literal:
-			{
-				Label = FString::Printf(TEXT("%s(Literal)"), *To<FRigVMLiteralExprAST>()->GetPin()->GetName());
-				break;
-			}
-			case EType::ExternalVar:
-			{
-				Label = FString::Printf(TEXT("%s(ExternalVar)"), *To<FRigVMExternalVarExprAST>()->GetPin()->GetBoundVariableName());
-				break;
-			}
-			case EType::Var:
-			{
-				if (To<FRigVMVarExprAST>()->IsGraphParameter())
-				{
-					URigVMParameterNode* ParameterNode = Cast<URigVMParameterNode>(To<FRigVMVarExprAST>()->GetPin()->GetNode());
-					check(ParameterNode);
-					Label = FString::Printf(TEXT("Param %s"), *ParameterNode->GetParameterName().ToString());
-				}
-				else if (To<FRigVMVarExprAST>()->IsGraphVariable())
-				{
-					URigVMVariableNode* VariableNode = Cast<URigVMVariableNode>(To<FRigVMVarExprAST>()->GetPin()->GetNode());
-					check(VariableNode);
-					Label = FString::Printf(TEXT("Variable %s"), *VariableNode->GetVariableName().ToString());
-				}
-				else if (To<FRigVMVarExprAST>()->IsEnumValue())
-				{
-					URigVMEnumNode* EnumNode = Cast<URigVMEnumNode>(To<FRigVMVarExprAST>()->GetPin()->GetNode());
-					check(EnumNode);
-					Label = FString::Printf(TEXT("Enum %s"), *EnumNode->GetCPPType());
-				}
-				else
-				{
-					Label = To<FRigVMVarExprAST>()->GetPin()->GetName();
-				}
-				if (To<FRigVMVarExprAST>()->IsExecuteContext())
-				{
-					AdditionalNodeSettings += TEXT(", shape = cds");
-				}
-				break;
-			}
-			case EType::Block:
-			{
-				if (GetParent() == nullptr)
-				{
-					Label = TEXT("Unused");
-					Result += FString::Printf(TEXT("\n%ssubgraph unused_%d {"), *Prefix, GetIndex());
-					Prefix += TEXT("  ");
-				}
-				else
-				{
-					Label = TEXT("Block");
-				}
-				break;
-			}
-			case EType::Assign:
-			{
-				Label = TEXT("=");
-				break;
-			}
-			case EType::Copy:
-			{
-				Label = TEXT("Copy");
-				break;
-			}
-			case EType::CachedValue:
-			{
-				Label = TEXT("Cache");
-				break;
-			}
-			case EType::CallExtern:
-			{
-				if (URigVMUnitNode* Node = Cast<URigVMUnitNode>(To<FRigVMCallExternExprAST>()->GetNode()))
-				{
-					Label = Node->GetScriptStruct()->GetName();
-				}
-				break;
-			}
-			case EType::NoOp:
-			{
-				Label = TEXT("NoOp");
-				break;
-			}
-			case EType::Array:
-			{
-				ERigVMOpCode OpCode = CastChecked<URigVMArrayNode>(To<FRigVMArrayExprAST>()->GetNode())->GetOpCode();
-				Label = StaticEnum<ERigVMOpCode>()->GetDisplayNameTextByValue((int64)OpCode).ToString();
-				break;
-			}
-			case EType::Exit:
-			{
-				Label = TEXT("Exit");
-				break;
-			}
-			case EType::Entry:
-			{
-				Result += FString::Printf(TEXT("\n%ssubgraph %s_%d {"), *Prefix, *GetName().ToString(), GetIndex());
-				Prefix += TEXT("  ");
-				break;
-			}
-			default:
-			{
-				break;
-			}
-		}
-
-		if (!Label.IsEmpty())
-		{
-			Result += FString::Printf(TEXT("\n%snode_%d [label = \"%s\"%s];"), *Prefix, GetIndex(), *Label, *AdditionalNodeSettings);
-		}
-
-		switch (GetType())
-		{
-			case EType::Entry:
-			case EType::Exit:
-			case EType::Branch:
-			case EType::Block:
-			{
-				Result += FString::Printf(TEXT("\n%snode_%d [shape = Mdiamond];"), *Prefix, GetIndex());
-				break;
-			}
-			case EType::Assign:
-			case EType::Copy:
-			case EType::CallExtern:
-			case EType::If:
-			case EType::Select:
-			case EType::NoOp:
-			{
-				Result += FString::Printf(TEXT("\n%snode_%d [shape = box];"), *Prefix, GetIndex());
-				break;
-			}
-			default:
-			{
-				break;
-			}
-		}
-
-	}
-
-	for (FRigVMExprAST* Child : Children)
-	{
-		Result += Child->DumpDot(OutExpressionDefined, Prefix);
-		if(!bWasDefined)
-		{
-			Result += FString::Printf(TEXT("\n%snode_%d -> node_%d;"), *Prefix, GetIndex(), Child->GetIndex());
-		}
-	}
-
-	if (!OutExpressionDefined[GetIndex()])
-	{
-		switch (GetType())
-		{
-			case EType::Block:
-			{
-				if (GetParent() == nullptr)
-				{
-					Prefix = Prefix.LeftChop(2);
-					Result += FString::Printf(TEXT("\n%s}"), *Prefix, *GetName().ToString(), GetIndex());
-				}
-				break;
-			}
-			case EType::Entry:
-			{
-				Prefix = Prefix.LeftChop(2);
-				Result += FString::Printf(TEXT("\n%s}"), *Prefix, *GetName().ToString(), GetIndex());
-				break;
-			}
-			default:
-			{
-				break;
-			}
-		}
-	}
-
-	OutExpressionDefined[GetIndex()] = true;
-
-	return Result;
-}
-
 bool FRigVMBlockExprAST::ShouldExecute() const
 {
 	return ContainsEntry();
@@ -2573,18 +2382,210 @@ FString FRigVMParserAST::DumpDot() const
 	TArray<bool> OutExpressionDefined;
 	OutExpressionDefined.AddZeroed(Expressions.Num());
 
-	FString Result = TEXT("digraph AST {\n  node [style=filled];\n  rankdir=\"LR\";");
+	FVisualGraph VisualGraph(TEXT("AST"));
 
-	for (FRigVMExprAST* RootExpr : RootExpressions)
+	VisualGraph.AddSubGraph(TEXT("AST"), FName(TEXT("AST")));
+	VisualGraph.AddSubGraph(TEXT("unused"), FName(TEXT("Unused")));
+
+	struct Local
+	{
+		static TArray<int32> VisitChildren(const FRigVMExprAST* InExpr, int32 InSubGraphIndex, FVisualGraph& OutGraph)
+		{
+			TArray<int32> ChildNodeIndices;
+			for (FRigVMExprAST* Child : InExpr->Children)
+			{
+				ChildNodeIndices.Add(VisitExpr(Child, InSubGraphIndex, OutGraph));
+			}
+			return ChildNodeIndices;
+		}
+		
+		static int32 VisitExpr(const FRigVMExprAST* InExpr, int32 InSubGraphIndex, FVisualGraph& OutGraph)
+		{
+			const FName NodeName = *FString::Printf(TEXT("node_%d"), InExpr->GetIndex());
+
+			int32 NodeIndex = OutGraph.FindNode(NodeName); 
+			if(NodeIndex != INDEX_NONE)
+			{
+				return NodeIndex;
+			}
+
+			FString Label = InExpr->GetName().ToString();
+			TOptional<EVisualGraphShape> Shape = EVisualGraphShape::Ellipse;
+			int32 SubGraphIndex = InSubGraphIndex;
+			
+			switch (InExpr->GetType())
+			{
+				case FRigVMExprAST::EType::Literal:
+				{
+					Label = FString::Printf(TEXT("%s(Literal)"), *InExpr->To<FRigVMLiteralExprAST>()->GetPin()->GetName());
+					break;
+				}
+				case FRigVMExprAST::EType::ExternalVar:
+				{
+					Label = FString::Printf(TEXT("%s(ExternalVar)"), *InExpr->To<FRigVMExternalVarExprAST>()->GetPin()->GetBoundVariableName());
+					break;
+				}
+				case FRigVMExprAST::EType::Var:
+				{
+					if (InExpr->To<FRigVMVarExprAST>()->IsGraphParameter())
+					{
+						URigVMParameterNode* ParameterNode = Cast<URigVMParameterNode>(InExpr->To<FRigVMVarExprAST>()->GetPin()->GetNode());
+						check(ParameterNode);
+						Label = FString::Printf(TEXT("Param %s"), *ParameterNode->GetParameterName().ToString());
+					}
+					else if (InExpr->To<FRigVMVarExprAST>()->IsGraphVariable())
+					{
+						URigVMVariableNode* VariableNode = Cast<URigVMVariableNode>(InExpr->To<FRigVMVarExprAST>()->GetPin()->GetNode());
+						check(VariableNode);
+						Label = FString::Printf(TEXT("Variable %s"), *VariableNode->GetVariableName().ToString());
+					}
+					else if (InExpr->To<FRigVMVarExprAST>()->IsEnumValue())
+					{
+						URigVMEnumNode* EnumNode = Cast<URigVMEnumNode>(InExpr->To<FRigVMVarExprAST>()->GetPin()->GetNode());
+						check(EnumNode);
+						Label = FString::Printf(TEXT("Enum %s"), *EnumNode->GetCPPType());
+					}
+					else
+					{
+						Label = InExpr->To<FRigVMVarExprAST>()->GetPin()->GetName();
+					}
+						
+					if (InExpr->To<FRigVMVarExprAST>()->IsExecuteContext())
+					{
+						Shape = EVisualGraphShape::House;
+					}
+						
+					break;
+				}
+				case FRigVMExprAST::EType::Block:
+				{
+					if (InExpr->GetParent() == nullptr)
+					{
+						Label = TEXT("Unused");
+						SubGraphIndex = OutGraph.FindSubGraph(TEXT("unused"));;
+					}
+					else
+					{
+						Label = TEXT("Block");
+					}
+					break;
+				}
+				case FRigVMExprAST::EType::Assign:
+				{
+					Label = TEXT("=");
+					break;
+				}
+				case FRigVMExprAST::EType::Copy:
+				{
+					Label = TEXT("Copy");
+					break;
+				}
+				case FRigVMExprAST::EType::CachedValue:
+				{
+					Label = TEXT("Cache");
+					break;
+				}
+				case FRigVMExprAST::EType::CallExtern:
+				{
+					if (URigVMUnitNode* Node = Cast<URigVMUnitNode>(InExpr->To<FRigVMCallExternExprAST>()->GetNode()))
+					{
+						Label = Node->GetScriptStruct()->GetName();
+					}
+					break;
+				}
+				case FRigVMExprAST::EType::NoOp:
+				{
+					Label = TEXT("NoOp");
+					break;
+				}
+				case FRigVMExprAST::EType::Array:
+				{
+					ERigVMOpCode OpCode = CastChecked<URigVMArrayNode>(InExpr->To<FRigVMArrayExprAST>()->GetNode())->GetOpCode();
+					Label = StaticEnum<ERigVMOpCode>()->GetDisplayNameTextByValue((int64)OpCode).ToString();
+					break;
+				}
+				case FRigVMExprAST::EType::Exit:
+				{
+					Label = TEXT("Exit");
+					break;
+				}
+				case FRigVMExprAST::EType::Entry:
+				{
+					SubGraphIndex = OutGraph.FindSubGraph(InExpr->GetName());
+					if(SubGraphIndex == INDEX_NONE)
+					{
+						const int32 ASTGraphIndex = OutGraph.FindSubGraph(TEXT("AST"));
+						SubGraphIndex = OutGraph.AddSubGraph(InExpr->GetName(), InExpr->GetName(), ASTGraphIndex);
+					}
+					break;
+				}
+				default:
+				{
+					break;
+				}
+			}
+
+			switch (InExpr->GetType())
+			{
+				case FRigVMExprAST::EType::Entry:
+				case FRigVMExprAST::EType::Exit:
+				case FRigVMExprAST::EType::Branch:
+				case FRigVMExprAST::EType::Block:
+				{
+					Shape = EVisualGraphShape::Diamond;
+					break;
+				}
+				case FRigVMExprAST::EType::Assign:
+				case FRigVMExprAST::EType::Copy:
+				case FRigVMExprAST::EType::CallExtern:
+				case FRigVMExprAST::EType::If:
+				case FRigVMExprAST::EType::Select:
+				case FRigVMExprAST::EType::NoOp:
+				{
+					Shape = EVisualGraphShape::Box;
+					break;
+				}
+				default:
+				{
+					break;
+				}
+			}
+
+			if (!Label.IsEmpty())
+			{
+				const TOptional<FName> DisplayName = FName(*Label);
+				NodeIndex = OutGraph.AddNode(NodeName, DisplayName, TOptional<FLinearColor>(), Shape);
+				OutGraph.AddNodeToSubGraph(NodeIndex, SubGraphIndex);
+			}
+
+			TArray<int32> ChildNodeIndices = VisitChildren(InExpr, SubGraphIndex, OutGraph);
+
+			if(NodeIndex != INDEX_NONE)
+			{
+				for(const int32 ChildNodeIndex : ChildNodeIndices)
+				{
+					if(ChildNodeIndex != INDEX_NONE)
+					{
+						OutGraph.AddEdge(ChildNodeIndex, NodeIndex, EVisualGraphEdgeDirection::SourceToTarget);
+					}
+				}
+			}
+
+			return NodeIndex;
+		}
+	};
+
+	for (FRigVMExprAST* Expr : RootExpressions)
 	{
-		if (RootExpr == GetObsoleteBlock(false /* create */))
+		if (Expr == GetObsoleteBlock(false))
 		{
 			continue;
 		}
-		Result += RootExpr->DumpDot(OutExpressionDefined, TEXT("  "));
+
+		Local::VisitExpr(Expr, INDEX_NONE, VisualGraph);
 	}
-	Result += TEXT("\n}");
-	return Result;
+
+	return VisualGraph.DumpDot();
 }
 
 FRigVMBlockExprAST* FRigVMParserAST::GetObsoleteBlock(bool bCreateIfMissing)
diff --git a/Engine/Source/Developer/RigVMDeveloper/Public/RigVMCompiler/RigVMAST.h b/Engine/Source/Developer/RigVMDeveloper/Public/RigVMCompiler/RigVMAST.h
--- a/Engine/Source/Developer/RigVMDeveloper/Public/RigVMCompiler/RigVMAST.h
+++ b/Engine/Source/Developer/RigVMDeveloper/Public/RigVMCompiler/RigVMAST.h
@@ -252,13 +252,6 @@ private:
 	// @return the text representation of this part of the tree
 	virtual FString DumpText(const FString& InPrefix = FString()) const;
 
-	// returns a string containing a dot file notation
-	// for debugging purposes. this is only used by the parser
-	// @param OutExpressionDefined a bool map to keep track of which expressions have been processed yet
-	// @param InPrefix the prefix to use for indentation
-	// @return the text representation of this part of the tree
-	virtual FString DumpDot(TArray<bool>& OutExpressionDefined, const FString& InPrefix = FString()) const;
-
 	FName Name;
 	EType Type;
 	int32 Index;
diff --git a/Engine/Source/Developer/RigVMDeveloper/RigVMDeveloper.Build.cs b/Engine/Source/Developer/RigVMDeveloper/RigVMDeveloper.Build.cs
--- a/Engine/Source/Developer/RigVMDeveloper/RigVMDeveloper.Build.cs
+++ b/Engine/Source/Developer/RigVMDeveloper/RigVMDeveloper.Build.cs
@@ -12,6 +12,7 @@ public class RigVMDeveloper : ModuleRules
                 "CoreUObject",
                 "Engine",
                 "RigVM",
+                "VisualGraphUtils",
             }
         );
 
diff --git a/Engine/Source/Developer/VisualGraphUtils/Private/VisualGraph.cpp b/Engine/Source/Developer/VisualGraphUtils/Private/VisualGraph.cpp
new file mode 100644
--- /dev/null
+++ b/Engine/Source/Developer/VisualGraphUtils/Private/VisualGraph.cpp
@@ -0,0 +1,360 @@
+﻿// Copyright Epic Games, Inc. All Rights Reserved.
+
+#include "VisualGraph.h"
+
+///////////////////////////////////////////////////////////////////////////////////
+/// FVisualGraphSubGraph
+///////////////////////////////////////////////////////////////////////////////////
+
+FString FVisualGraphSubGraph::DumpDot(const FVisualGraph* InGraph, int32 InIndendation) const
+{
+	FString EdgeContent;
+
+	if(!Nodes.IsEmpty())
+	{
+		for(const FVisualGraphEdge& Edge : InGraph->GetEdges())
+		{
+			if(Nodes.Contains(Edge.GetSourceNode()) &&
+				Nodes.Contains(Edge.GetTargetNode()))
+			{
+				EdgeContent += Edge.DumpDot(InGraph, InIndendation + 1);
+			}
+		}
+	}
+
+	FString SubGraphContent;
+	for(const FVisualGraphSubGraph& SubGraph : InGraph->GetSubGraphs())
+	{
+		if(SubGraph.GetParentGraphIndex() == GetIndex())
+		{
+			SubGraphContent += SubGraph.DumpDot(InGraph, InIndendation + 1);
+		}
+	}
+
+	if(EdgeContent.IsEmpty() && SubGraphContent.IsEmpty())
+	{
+		return FString();
+	}
+
+	const FString Indentation1 = DumpDotIndentation(InIndendation);
+	const FString Indentation2 = DumpDotIndentation(InIndendation + 1);
+	const FString ColorContent = DumpDotColor(GetColor());
+	const FString StyleContent = DumpDotStyle(GetStyle());
+
+	FString AttributeContent;
+	if(DisplayName.IsSet())
+	{
+		AttributeContent += FString::Printf(TEXT("%slabel = \"%s\";\n"),
+			*Indentation2, *DisplayName.GetValue().ToString());
+	}
+	if(!ColorContent.IsEmpty())
+	{
+		AttributeContent += FString::Printf(TEXT("%scolor = %s;\n"),
+			*Indentation2, *ColorContent);
+	}
+	if(!StyleContent.IsEmpty())
+	{
+		AttributeContent += FString::Printf(TEXT("%sstyle = %s;\n"),
+			*Indentation2, *StyleContent);
+	}
+
+	return FString::Printf(TEXT("%ssubgraph cluster_%s {\n%s%s%s%s}\n"),
+		*Indentation1, *GetName().ToString(),
+		*AttributeContent,
+		*EdgeContent,
+		*SubGraphContent,
+		*Indentation1
+		);
+}
+
+///////////////////////////////////////////////////////////////////////////////////
+/// FVisualGraph
+///////////////////////////////////////////////////////////////////////////////////
+
+FVisualGraph::FVisualGraph(const FName& InName, const FName& InDisplayName)
+{
+	Name = InName;
+	DisplayName = InDisplayName;
+	Style = EVisualGraphStyle::Filled;
+}
+
+int32 FVisualGraph::AddNode(const FName& InName, TOptional<FName> InDisplayName, TOptional<FLinearColor> InColor,
+	TOptional<EVisualGraphShape> InShape, TOptional<EVisualGraphStyle> InStyle)
+{
+	RefreshNameMapIfNeeded(Nodes, NodeNameMap);
+	
+	const FName UniqueName = GetUniqueName(InName, NodeNameMap);
+	FVisualGraphNode Node;
+	Node.Name = UniqueName;
+	Node.DisplayName = InDisplayName;
+	Node.Color = InColor;
+	Node.Shape = InShape;
+	Node.Style = InStyle;
+	Node.SubGraphIndex = INDEX_NONE;
+
+	return AddElement(Node, Nodes, NodeNameMap);
+}
+
+int32 FVisualGraph::AddEdge(int32 InSourceNode, int32 InTargetNode, EVisualGraphEdgeDirection InDirection,
+	const FName& InName, TOptional<FName> InDisplayName, TOptional<FLinearColor> InColor,
+	TOptional<EVisualGraphStyle> InStyle)
+{
+	RefreshNameMapIfNeeded(Edges, EdgeNameMap);
+
+	FName DesiredName = InName;
+	if(DesiredName.IsNone())
+	{
+		DesiredName = TEXT("Edge");
+	}
+	
+	const FName UniqueName = GetUniqueName(DesiredName, EdgeNameMap);
+	FVisualGraphEdge Edge;
+	Edge.Name = UniqueName;
+	Edge.DisplayName = InDisplayName;
+	Edge.Color = InColor;
+	Edge.Style = InStyle;
+	Edge.SourceNode = InSourceNode;
+	Edge.TargetNode = InTargetNode;
+	Edge.Direction = InDirection;
+
+	return AddElement(Edge, Edges, EdgeNameMap);
+}
+
+int32 FVisualGraph::AddSubGraph(const FName& InName, TOptional<FName> InDisplayName, int32 InParentGraphIndex,
+	TOptional<FLinearColor> InColor, TOptional<EVisualGraphStyle> InStyle, const TArray<int32> InNodes)
+{
+	RefreshNameMapIfNeeded(SubGraphs, SubGraphNameMap);
+
+	const FName UniqueName = GetUniqueName(InName, SubGraphNameMap);
+	FVisualGraphSubGraph SubGraph;
+	SubGraph.Name = UniqueName;
+	SubGraph.DisplayName = InDisplayName;
+	SubGraph.Color = InColor;
+	SubGraph.Style = InStyle;
+	SubGraph.ParentGraphIndex = InParentGraphIndex;
+
+	const int32 SubGraphIndex = AddElement(SubGraph, SubGraphs, SubGraphNameMap);
+
+	for(const int32 Node : InNodes)
+	{
+		AddNodeToSubGraph(Node, SubGraphIndex);
+	}
+
+	return SubGraphIndex;
+}
+
+int32 FVisualGraph::FindNode(const FName& InName) const
+{
+	if(const int32* IndexPtr= NodeNameMap.Find(InName))
+	{
+		return *IndexPtr;
+	}
+	return INDEX_NONE;
+}
+
+int32 FVisualGraph::FindEdge(const FName& InName) const
+{
+	if(const int32* IndexPtr= EdgeNameMap.Find(InName))
+	{
+		return *IndexPtr;
+	}
+	return INDEX_NONE;
+}
+
+int32 FVisualGraph::FindSubGraph(const FName& InName) const
+{
+	if(const int32* IndexPtr= SubGraphNameMap.Find(InName))
+	{
+		return *IndexPtr;
+	}
+	return INDEX_NONE;
+}
+
+bool FVisualGraph::AddNodeToSubGraph(int32 InNodeIndex, int32 InSubGraphIndex)
+{
+	if(!Nodes.IsValidIndex(InNodeIndex))
+	{
+		return false;
+	}
+
+	if(!SubGraphs.IsValidIndex(InSubGraphIndex))
+	{
+		return false;
+	}
+
+	RemoveNodeFromSubGraph(InNodeIndex);
+	SubGraphs[InSubGraphIndex].Nodes.Add(InNodeIndex);
+	Nodes[InNodeIndex].SubGraphIndex = InSubGraphIndex;
+	return true;
+}
+
+bool FVisualGraph::RemoveNodeFromSubGraph(int32 InNodeIndex)
+{
+	if(!Nodes.IsValidIndex(InNodeIndex))
+	{
+		return false;
+	}
+	if(Nodes[InNodeIndex].SubGraphIndex == INDEX_NONE)
+	{
+		return false;
+	}
+	if(!SubGraphs.IsValidIndex(Nodes[InNodeIndex].SubGraphIndex))
+	{
+		return false;
+	}
+
+	SubGraphs[Nodes[InNodeIndex].SubGraphIndex].Nodes.Remove(InNodeIndex);
+	Nodes[InNodeIndex].SubGraphIndex = INDEX_NONE;
+	return true;
+}
+
+FString FVisualGraph::DumpDot(const FVisualGraph* InGraph, int32 InIndendation) const
+{
+	check(InGraph == this);
+
+	const FString Indentation = DumpDotIndentation(InIndendation);
+	const FString Indentation2 = DumpDotIndentation(InIndendation + 1);
+
+	const FString ColorContent = DumpDotColor(GetColor());
+	const FString StyleContent = DumpDotStyle(GetStyle());
+	FString AttributeContent;
+	if(!ColorContent.IsEmpty())
+	{
+		if(!AttributeContent.IsEmpty())
+		{
+			AttributeContent += TEXT(", ");
+		}
+		AttributeContent += FString::Printf(TEXT("color = %s"), *ColorContent);
+	}
+	if(!StyleContent.IsEmpty())
+	{
+		if(!AttributeContent.IsEmpty())
+		{
+			AttributeContent += TEXT(", ");
+		}
+		AttributeContent += FString::Printf(TEXT("style = %s"), *StyleContent);
+	}
+	if(!AttributeContent.IsEmpty())
+	{
+		AttributeContent = FString::Printf(TEXT("%snode [ %s ];\n"), *Indentation2, *AttributeContent);
+	}
+
+	FString SubGraphContent;
+	for(const FVisualGraphSubGraph& SubGraph : SubGraphs)
+	{
+		if(SubGraph.GetParentGraphIndex() == INDEX_NONE)
+		{
+			SubGraphContent += SubGraph.DumpDot(this, InIndendation + 1);
+		}
+	}
+
+	bool bIsDirectedGraph = true;
+	FString EdgeContent;
+	for(const FVisualGraphEdge& Edge : Edges)
+	{
+		if(Edge.GetDirection() == EVisualGraphEdgeDirection::BothWays)
+		{
+			bIsDirectedGraph = false;
+		}
+
+		const FVisualGraphNode& SourceNode = Nodes[Edge.GetSourceNode()];
+		const FVisualGraphNode& TargetNode = Nodes[Edge.GetTargetNode()];
+
+		// if both source and target node are within the same subgraph
+		if((SourceNode.SubGraphIndex == TargetNode.SubGraphIndex) &&
+			(SourceNode.SubGraphIndex != INDEX_NONE))
+		{
+			continue;
+		}
+		
+		EdgeContent += Edge.DumpDot(this, InIndendation + 1);
+	}
+
+	FString NodeContent;
+	for(const FVisualGraphNode& Node : Nodes)
+	{
+		NodeContent += Node.DumpDot(this, InIndendation + 1);
+	}
+
+	const FString GraphToken = bIsDirectedGraph ? TEXT("digraph") : TEXT("graph"); 
+	return FString::Printf(TEXT("%s%s %s {\n%srankdir=\"LR\";\n%s%s%s%s%s}\n"),
+		*Indentation, *GraphToken, *GetName().ToString(),
+		*Indentation2,
+		*AttributeContent,
+		*SubGraphContent,
+		*EdgeContent,
+		*NodeContent,
+		*Indentation
+		);
+}
+
+bool FVisualGraph::IsNameAvailable(const FName& InName, const TMap<FName, int32>& InMap)
+{
+	return !InMap.Contains(InName);
+}
+
+FName FVisualGraph::GetUniqueName(const FName& InName, const TMap<FName, int32>& InMap)
+{
+	const FName DesiredName = InName;
+	FName Name = DesiredName;
+	int32 Suffix = 0;
+	while (!IsNameAvailable(Name, InMap))
+	{
+		Suffix++;
+		Name = *FString::Printf(TEXT("%s_%d"), *DesiredName.ToString(), Suffix);
+	}
+	return Name;
+}
+
+void FVisualGraph::TransitiveReduction(TFunction<bool(FVisualGraphEdge&)> KeepEdgeFunction)
+{
+	// mark up all valid edges
+	TArray<int32> ValidEdge;
+	ValidEdge.AddUninitialized(Nodes.Num() * Nodes.Num());
+	for(int32 EdgeIndex = 0; EdgeIndex < ValidEdge.Num(); EdgeIndex++)
+	{
+		ValidEdge[EdgeIndex] = -1;
+	}
+	for(const FVisualGraphEdge& Edge : Edges)
+	{
+		const int32 EdgeIndex = Edge.GetSourceNode() * Nodes.Num() + Edge.GetTargetNode();
+		ValidEdge[EdgeIndex] = Edge.GetIndex();
+	}
+
+
+	TArray<int32> EdgesToRemove;
+
+	// perform transitive reduction on the graph
+	for(const FVisualGraphNode& NodeA : Nodes)
+	{
+		for(const FVisualGraphNode& NodeB : Nodes)
+		{
+			for(const FVisualGraphNode& NodeC : Nodes)
+			{
+				const int32 EdgeIndexAB = NodeA.GetIndex() * Nodes.Num() + NodeB.GetIndex();
+				const int32 EdgeIndexBC = NodeB.GetIndex() * Nodes.Num() + NodeC.GetIndex();
+				const int32 EdgeIndexAC = NodeA.GetIndex() * Nodes.Num() + NodeC.GetIndex();
+
+				if(ValidEdge[EdgeIndexAC] != INDEX_NONE &&
+					ValidEdge[EdgeIndexAB] != INDEX_NONE &&
+					ValidEdge[EdgeIndexBC] != INDEX_NONE)
+				{
+					const int32 EdgeIndex = ValidEdge[EdgeIndexAC];
+					if(!KeepEdgeFunction(Edges[EdgeIndex]))
+					{
+						EdgesToRemove.AddUnique(EdgeIndex);
+					}
+					ValidEdge[EdgeIndexAC] = INDEX_NONE;
+				}
+			}
+		}
+	}
+
+	Algo::Sort(EdgesToRemove);
+	Algo::Reverse(EdgesToRemove);
+
+	for(const int32 EdgeToRemove : EdgesToRemove)
+	{
+		Edges.RemoveAt(EdgeToRemove);
+	}
+}
\ No newline at end of file
diff --git a/Engine/Source/Developer/VisualGraphUtils/Private/VisualGraphEdge.cpp b/Engine/Source/Developer/VisualGraphUtils/Private/VisualGraphEdge.cpp
new file mode 100644
--- /dev/null
+++ b/Engine/Source/Developer/VisualGraphUtils/Private/VisualGraphEdge.cpp
@@ -0,0 +1,64 @@
+﻿
+#include "VisualGraphEdge.h"
+#include "VisualGraph.h"
+
+///////////////////////////////////////////////////////////////////////////////////
+/// FVisualGraphEdge
+///////////////////////////////////////////////////////////////////////////////////
+
+FString FVisualGraphEdge::DumpDot(const FVisualGraph* InGraph, int32 InIndendation) const
+{
+	const FString Indentation = DumpDotIndentation(InIndendation);
+	const FString ColorContent = DumpDotColor(GetColor());
+	const FString StyleContent = DumpDotStyle(GetStyle());
+
+	static const FString DigraphConnector = TEXT("->");
+	static const FString GraphConnector = TEXT("--");
+	const FString& ConnectorContent = (GetDirection() == EVisualGraphEdgeDirection::BothWays) ? GraphConnector : DigraphConnector; 
+
+	FString SourceNodeName = InGraph->GetNodes()[GetSourceNode()].GetName().ToString();
+	FString TargetNodeName = InGraph->GetNodes()[GetTargetNode()].GetName().ToString();
+
+	if(GetDirection() == EVisualGraphEdgeDirection::SourceToTarget)
+	{
+		Swap(SourceNodeName, TargetNodeName);
+	}
+
+	FString AttributeContent;
+	if(DisplayName.IsSet())
+	{
+		AttributeContent += FString::Printf(TEXT("label = \"%s\""), *DisplayName.GetValue().ToString());
+	}
+	if(Tooltip.IsSet())
+	{
+		if(!AttributeContent.IsEmpty())
+		{
+			AttributeContent += TEXT(", ");
+		}
+		AttributeContent += FString::Printf(TEXT("tooltip = \"%s\""), *Tooltip.GetValue());
+	}
+	if(!ColorContent.IsEmpty())
+	{
+		if(!AttributeContent.IsEmpty())
+		{
+			AttributeContent += TEXT(", ");
+		}
+		AttributeContent += FString::Printf(TEXT("color = %s"), *ColorContent);
+	}
+	if(!StyleContent.IsEmpty())
+	{
+		if(!AttributeContent.IsEmpty())
+		{
+			AttributeContent += TEXT(", ");
+		}
+		AttributeContent += FString::Printf(TEXT("style = %s"), *StyleContent);
+	}
+
+	if(!AttributeContent.IsEmpty())
+	{
+		AttributeContent = FString::Printf(TEXT(" [ %s ]"), *AttributeContent);
+	}
+
+	return FString::Printf(TEXT("%s%s %s %s%s;\n"),
+		*Indentation, *SourceNodeName, *ConnectorContent, *TargetNodeName, *AttributeContent);
+}
\ No newline at end of file
diff --git a/Engine/Source/Developer/VisualGraphUtils/Private/VisualGraphElement.cpp b/Engine/Source/Developer/VisualGraphUtils/Private/VisualGraphElement.cpp
new file mode 100644
--- /dev/null
+++ b/Engine/Source/Developer/VisualGraphUtils/Private/VisualGraphElement.cpp
@@ -0,0 +1,223 @@
+﻿// Copyright Epic Games, Inc. All Rights Reserved.
+
+#include "VisualGraphElement.h"
+#include "VisualGraph.h"
+
+///////////////////////////////////////////////////////////////////////////////////
+/// FVisualGraphElement
+///////////////////////////////////////////////////////////////////////////////////
+
+FString FVisualGraphElement::DumpDotIndentation(int32 InIndentation)
+{
+	check(InIndentation >= 0);
+	
+	static TArray<FString> Indentations;
+	while(Indentations.Num() <= InIndentation)
+	{
+		const int32 Indentation = Indentations.Num();
+		if(Indentation == 0)
+		{
+			Indentations.Add(FString());
+		}
+		else
+		{
+			Indentations.Add(Indentations[Indentation - 1] + TEXT("  "));
+		}
+	}
+
+	return Indentations[InIndentation];
+}
+
+FString FVisualGraphElement::DumpDotColor(const TOptional<FLinearColor>& InColor)
+{
+	if(!InColor.IsSet())
+	{
+		return FString();
+	}
+
+	const FLinearColor Color = InColor.GetValue();
+	
+	static TMap<FName, FLinearColor> ColorScheme;
+	if(ColorScheme.IsEmpty())
+	{
+		ColorScheme.Add(TEXT("alice blue"), FColor::FromHex(TEXT("#F0F8FF")));
+		ColorScheme.Add(TEXT("antique white"), FColor::FromHex(TEXT("#FAEBD7")));
+		//ColorScheme.Add(TEXT("aqua"), FColor::FromHex(TEXT("#00FFFF")));
+		ColorScheme.Add(TEXT("aquamarine"), FColor::FromHex(TEXT("#7FFFD4")));
+		ColorScheme.Add(TEXT("azure"), FColor::FromHex(TEXT("#F0FFFF")));
+		ColorScheme.Add(TEXT("beige"), FColor::FromHex(TEXT("#F5F5DC")));
+		ColorScheme.Add(TEXT("bisque"), FColor::FromHex(TEXT("#FFE4C4")));
+		ColorScheme.Add(TEXT("black"), FColor::FromHex(TEXT("#000000")));
+		ColorScheme.Add(TEXT("blanched almond"), FColor::FromHex(TEXT("#FFEBCD")));
+		ColorScheme.Add(TEXT("blue"), FColor::FromHex(TEXT("#0000FF")));
+		ColorScheme.Add(TEXT("blue violet"), FColor::FromHex(TEXT("#8A2BE2")));
+		ColorScheme.Add(TEXT("brown"), FColor::FromHex(TEXT("#A52A2A")));
+		ColorScheme.Add(TEXT("burlywood"), FColor::FromHex(TEXT("#DEB887")));
+		ColorScheme.Add(TEXT("cadet blue"), FColor::FromHex(TEXT("#5F9EA0")));
+		ColorScheme.Add(TEXT("chartreuse"), FColor::FromHex(TEXT("#7FFF00")));
+		ColorScheme.Add(TEXT("chocolate"), FColor::FromHex(TEXT("#D2691E")));
+		ColorScheme.Add(TEXT("coral"), FColor::FromHex(TEXT("#FF7F50")));
+		ColorScheme.Add(TEXT("cornflower blue"), FColor::FromHex(TEXT("#6495ED")));
+		ColorScheme.Add(TEXT("cornsilk"), FColor::FromHex(TEXT("#FFF8DC")));
+		ColorScheme.Add(TEXT("crimson"), FColor::FromHex(TEXT("#DC143C")));
+		ColorScheme.Add(TEXT("cyan"), FColor::FromHex(TEXT("#00FFFF")));
+		ColorScheme.Add(TEXT("dark blue"), FColor::FromHex(TEXT("#00008B")));
+		ColorScheme.Add(TEXT("dark cyan"), FColor::FromHex(TEXT("#008B8B")));
+		ColorScheme.Add(TEXT("dark goldenrod"), FColor::FromHex(TEXT("#B8860B")));
+		ColorScheme.Add(TEXT("dark gray"), FColor::FromHex(TEXT("#A9A9A9")));
+		ColorScheme.Add(TEXT("dark green"), FColor::FromHex(TEXT("#006400")));
+		ColorScheme.Add(TEXT("dark khaki"), FColor::FromHex(TEXT("#BDB76B")));
+		ColorScheme.Add(TEXT("dark magenta"), FColor::FromHex(TEXT("#8B008B")));
+		ColorScheme.Add(TEXT("dark olive green"), FColor::FromHex(TEXT("#556B2F")));
+		ColorScheme.Add(TEXT("dark orange"), FColor::FromHex(TEXT("#FF8C00")));
+		ColorScheme.Add(TEXT("dark orchid"), FColor::FromHex(TEXT("#9932CC")));
+		ColorScheme.Add(TEXT("dark red"), FColor::FromHex(TEXT("#8B0000")));
+		ColorScheme.Add(TEXT("dark salmon"), FColor::FromHex(TEXT("#E9967A")));
+		ColorScheme.Add(TEXT("dark sea green"), FColor::FromHex(TEXT("#8FBC8F")));
+		ColorScheme.Add(TEXT("dark slate blue"), FColor::FromHex(TEXT("#483D8B")));
+		ColorScheme.Add(TEXT("dark slate gray"), FColor::FromHex(TEXT("#2F4F4F")));
+		ColorScheme.Add(TEXT("dark turquoise"), FColor::FromHex(TEXT("#00CED1")));
+		ColorScheme.Add(TEXT("dark violet"), FColor::FromHex(TEXT("#9400D3")));
+		ColorScheme.Add(TEXT("deep pink"), FColor::FromHex(TEXT("#FF1493")));
+		ColorScheme.Add(TEXT("deep sky blue"), FColor::FromHex(TEXT("#00BFFF")));
+		ColorScheme.Add(TEXT("dim gray"), FColor::FromHex(TEXT("#696969")));
+		ColorScheme.Add(TEXT("dodger blue"), FColor::FromHex(TEXT("#1E90FF")));
+		ColorScheme.Add(TEXT("firebrick"), FColor::FromHex(TEXT("#B22222")));
+		ColorScheme.Add(TEXT("floral white"), FColor::FromHex(TEXT("#FFFAF0")));
+		ColorScheme.Add(TEXT("forest green"), FColor::FromHex(TEXT("#228B22")));
+		ColorScheme.Add(TEXT("fuchsia"), FColor::FromHex(TEXT("#FF00FF")));
+		ColorScheme.Add(TEXT("gainsboro"), FColor::FromHex(TEXT("#DCDCDC")));
+		ColorScheme.Add(TEXT("ghost white"), FColor::FromHex(TEXT("#F8F8FF")));
+		ColorScheme.Add(TEXT("gold"), FColor::FromHex(TEXT("#FFD700")));
+		ColorScheme.Add(TEXT("goldenrod"), FColor::FromHex(TEXT("#DAA520")));
+		ColorScheme.Add(TEXT("gray"), FColor::FromHex(TEXT("#BEBEBE")));
+		ColorScheme.Add(TEXT("web gray"), FColor::FromHex(TEXT("#808080")));
+		ColorScheme.Add(TEXT("green"), FColor::FromHex(TEXT("#00FF00")));
+		ColorScheme.Add(TEXT("web green"), FColor::FromHex(TEXT("#008000")));
+		ColorScheme.Add(TEXT("green yellow"), FColor::FromHex(TEXT("#ADFF2F")));
+		ColorScheme.Add(TEXT("honeydew"), FColor::FromHex(TEXT("#F0FFF0")));
+		ColorScheme.Add(TEXT("hot pink"), FColor::FromHex(TEXT("#FF69B4")));
+		ColorScheme.Add(TEXT("indian red"), FColor::FromHex(TEXT("#CD5C5C")));
+		ColorScheme.Add(TEXT("indigo"), FColor::FromHex(TEXT("#4B0082")));
+		ColorScheme.Add(TEXT("ivory"), FColor::FromHex(TEXT("#FFFFF0")));
+		ColorScheme.Add(TEXT("khaki"), FColor::FromHex(TEXT("#F0E68C")));
+		ColorScheme.Add(TEXT("lavender"), FColor::FromHex(TEXT("#E6E6FA")));
+		ColorScheme.Add(TEXT("lavender blush"), FColor::FromHex(TEXT("#FFF0F5")));
+		ColorScheme.Add(TEXT("lawn green"), FColor::FromHex(TEXT("#7CFC00")));
+		ColorScheme.Add(TEXT("lemon chiffon"), FColor::FromHex(TEXT("#FFFACD")));
+		ColorScheme.Add(TEXT("light blue"), FColor::FromHex(TEXT("#ADD8E6")));
+		ColorScheme.Add(TEXT("light coral"), FColor::FromHex(TEXT("#F08080")));
+		ColorScheme.Add(TEXT("light cyan"), FColor::FromHex(TEXT("#E0FFFF")));
+		ColorScheme.Add(TEXT("light goldenrod"), FColor::FromHex(TEXT("#FAFAD2")));
+		ColorScheme.Add(TEXT("light gray"), FColor::FromHex(TEXT("#D3D3D3")));
+		ColorScheme.Add(TEXT("light green"), FColor::FromHex(TEXT("#90EE90")));
+		ColorScheme.Add(TEXT("light pink"), FColor::FromHex(TEXT("#FFB6C1")));
+		ColorScheme.Add(TEXT("light salmon"), FColor::FromHex(TEXT("#FFA07A")));
+		ColorScheme.Add(TEXT("light sea green"), FColor::FromHex(TEXT("#20B2AA")));
+		ColorScheme.Add(TEXT("light sky blue"), FColor::FromHex(TEXT("#87CEFA")));
+		ColorScheme.Add(TEXT("light slate gray"), FColor::FromHex(TEXT("#778899")));
+		ColorScheme.Add(TEXT("light steel blue"), FColor::FromHex(TEXT("#B0C4DE")));
+		ColorScheme.Add(TEXT("light yellow"), FColor::FromHex(TEXT("#FFFFE0")));
+		ColorScheme.Add(TEXT("lime"), FColor::FromHex(TEXT("#00FF00")));
+		ColorScheme.Add(TEXT("lime green"), FColor::FromHex(TEXT("#32CD32")));
+		ColorScheme.Add(TEXT("linen"), FColor::FromHex(TEXT("#FAF0E6")));
+		ColorScheme.Add(TEXT("magenta"), FColor::FromHex(TEXT("#FF00FF")));
+		ColorScheme.Add(TEXT("maroon"), FColor::FromHex(TEXT("#B03060")));
+		ColorScheme.Add(TEXT("web maroon"), FColor::FromHex(TEXT("#800000")));
+		ColorScheme.Add(TEXT("medium aquamarine"), FColor::FromHex(TEXT("#66CDAA")));
+		ColorScheme.Add(TEXT("medium blue"), FColor::FromHex(TEXT("#0000CD")));
+		ColorScheme.Add(TEXT("medium orchid"), FColor::FromHex(TEXT("#BA55D3")));
+		ColorScheme.Add(TEXT("medium purple"), FColor::FromHex(TEXT("#9370DB")));
+		ColorScheme.Add(TEXT("medium sea green"), FColor::FromHex(TEXT("#3CB371")));
+		ColorScheme.Add(TEXT("medium slate blue"), FColor::FromHex(TEXT("#7B68EE")));
+		ColorScheme.Add(TEXT("medium spring green"), FColor::FromHex(TEXT("#00FA9A")));
+		ColorScheme.Add(TEXT("medium turquoise"), FColor::FromHex(TEXT("#48D1CC")));
+		ColorScheme.Add(TEXT("medium violet red"), FColor::FromHex(TEXT("#C71585")));
+		ColorScheme.Add(TEXT("midnight blue"), FColor::FromHex(TEXT("#191970")));
+		ColorScheme.Add(TEXT("mint cream"), FColor::FromHex(TEXT("#F5FFFA")));
+		ColorScheme.Add(TEXT("misty rose"), FColor::FromHex(TEXT("#FFE4E1")));
+		ColorScheme.Add(TEXT("moccasin"), FColor::FromHex(TEXT("#FFE4B5")));
+		ColorScheme.Add(TEXT("navajo white"), FColor::FromHex(TEXT("#FFDEAD")));
+		ColorScheme.Add(TEXT("navy blue"), FColor::FromHex(TEXT("#000080")));
+		ColorScheme.Add(TEXT("old lace"), FColor::FromHex(TEXT("#FDF5E6")));
+		ColorScheme.Add(TEXT("olive"), FColor::FromHex(TEXT("#808000")));
+		ColorScheme.Add(TEXT("olive drab"), FColor::FromHex(TEXT("#6B8E23")));
+		ColorScheme.Add(TEXT("orange"), FColor::FromHex(TEXT("#FFA500")));
+		ColorScheme.Add(TEXT("orange red"), FColor::FromHex(TEXT("#FF4500")));
+		ColorScheme.Add(TEXT("orchid"), FColor::FromHex(TEXT("#DA70D6")));
+		ColorScheme.Add(TEXT("pale goldenrod"), FColor::FromHex(TEXT("#EEE8AA")));
+		ColorScheme.Add(TEXT("pale green"), FColor::FromHex(TEXT("#98FB98")));
+		ColorScheme.Add(TEXT("pale turquoise"), FColor::FromHex(TEXT("#AFEEEE")));
+		ColorScheme.Add(TEXT("pale violet red"), FColor::FromHex(TEXT("#DB7093")));
+		ColorScheme.Add(TEXT("papaya whip"), FColor::FromHex(TEXT("#FFEFD5")));
+		ColorScheme.Add(TEXT("peach puff"), FColor::FromHex(TEXT("#FFDAB9")));
+		ColorScheme.Add(TEXT("peru"), FColor::FromHex(TEXT("#CD853F")));
+		ColorScheme.Add(TEXT("pink"), FColor::FromHex(TEXT("#FFC0CB")));
+		ColorScheme.Add(TEXT("plum"), FColor::FromHex(TEXT("#DDA0DD")));
+		ColorScheme.Add(TEXT("powder blue"), FColor::FromHex(TEXT("#B0E0E6")));
+		ColorScheme.Add(TEXT("purple"), FColor::FromHex(TEXT("#A020F0")));
+		ColorScheme.Add(TEXT("web purple"), FColor::FromHex(TEXT("#800080")));
+		ColorScheme.Add(TEXT("rebecca purple"), FColor::FromHex(TEXT("#663399")));
+		ColorScheme.Add(TEXT("red"), FColor::FromHex(TEXT("#FF0000")));
+		ColorScheme.Add(TEXT("rosy brown"), FColor::FromHex(TEXT("#BC8F8F")));
+		ColorScheme.Add(TEXT("royal blue"), FColor::FromHex(TEXT("#4169E1")));
+		ColorScheme.Add(TEXT("saddle brown"), FColor::FromHex(TEXT("#8B4513")));
+		ColorScheme.Add(TEXT("salmon"), FColor::FromHex(TEXT("#FA8072")));
+		ColorScheme.Add(TEXT("sandy brown"), FColor::FromHex(TEXT("#F4A460")));
+		ColorScheme.Add(TEXT("sea green"), FColor::FromHex(TEXT("#2E8B57")));
+		ColorScheme.Add(TEXT("seashell"), FColor::FromHex(TEXT("#FFF5EE")));
+		ColorScheme.Add(TEXT("sienna"), FColor::FromHex(TEXT("#A0522D")));
+		ColorScheme.Add(TEXT("silver"), FColor::FromHex(TEXT("#C0C0C0")));
+		ColorScheme.Add(TEXT("sky blue"), FColor::FromHex(TEXT("#87CEEB")));
+		ColorScheme.Add(TEXT("slate blue"), FColor::FromHex(TEXT("#6A5ACD")));
+		ColorScheme.Add(TEXT("slate gray"), FColor::FromHex(TEXT("#708090")));
+		ColorScheme.Add(TEXT("snow"), FColor::FromHex(TEXT("#FFFAFA")));
+		ColorScheme.Add(TEXT("spring green"), FColor::FromHex(TEXT("#00FF7F")));
+		ColorScheme.Add(TEXT("steel blue"), FColor::FromHex(TEXT("#4682B4")));
+		ColorScheme.Add(TEXT("tan"), FColor::FromHex(TEXT("#D2B48C")));
+		ColorScheme.Add(TEXT("teal"), FColor::FromHex(TEXT("#008080")));
+		ColorScheme.Add(TEXT("thistle"), FColor::FromHex(TEXT("#D8BFD8")));
+		ColorScheme.Add(TEXT("tomato"), FColor::FromHex(TEXT("#FF6347")));
+		ColorScheme.Add(TEXT("turquoise"), FColor::FromHex(TEXT("#40E0D0")));
+		ColorScheme.Add(TEXT("violet"), FColor::FromHex(TEXT("#EE82EE")));
+		ColorScheme.Add(TEXT("wheat"), FColor::FromHex(TEXT("#F5DEB3")));
+		ColorScheme.Add(TEXT("white"), FColor::FromHex(TEXT("#FFFFFF")));
+		ColorScheme.Add(TEXT("white smoke"), FColor::FromHex(TEXT("#F5F5F5")));
+		ColorScheme.Add(TEXT("yellow"), FColor::FromHex(TEXT("#FFFF00")));
+		ColorScheme.Add(TEXT("yellow green"), FColor::FromHex(TEXT("#9ACD32")));
+	}
+
+	FName SchemeName = NAME_None;
+	float MinDelta = FLT_MAX;
+	for(const TPair<FName, FLinearColor>& Pair : ColorScheme)
+	{
+		const float Delta = FLinearColor::Dist(Pair.Value, Color);
+		if(MinDelta > Delta)
+		{
+			MinDelta = Delta;
+			SchemeName = Pair.Key;
+		}
+	}
+
+	FString SchemeString = SchemeName.ToString();
+	SchemeString.RemoveSpacesInline();
+	return SchemeString;
+}
+
+FString FVisualGraphElement::DumpDotShape(const TOptional<EVisualGraphShape>& InShape)
+{
+	if(!InShape.IsSet())
+	{
+		return FString();
+	}
+	return StaticEnum<EVisualGraphShape>()->GetDisplayNameTextByValue((int64)InShape.GetValue()).ToString().ToLower();
+}
+
+FString FVisualGraphElement::DumpDotStyle(const TOptional<EVisualGraphStyle>& InStyle)
+{
+	if(!InStyle.IsSet())
+	{
+		return FString();
+	}
+	return StaticEnum<EVisualGraphStyle>()->GetDisplayNameTextByValue((int64)InStyle.GetValue()).ToString().ToLower();
+}
\ No newline at end of file
diff --git a/Engine/Source/Developer/VisualGraphUtils/Private/VisualGraphNode.cpp b/Engine/Source/Developer/VisualGraphUtils/Private/VisualGraphNode.cpp
new file mode 100644
--- /dev/null
+++ b/Engine/Source/Developer/VisualGraphUtils/Private/VisualGraphNode.cpp
@@ -0,0 +1,75 @@
+﻿// Copyright Epic Games, Inc. All Rights Reserved.
+
+#include "VisualGraphNode.h"
+#include "VisualGraph.h"
+
+///////////////////////////////////////////////////////////////////////////////////
+/// FVisualGraphNode
+///////////////////////////////////////////////////////////////////////////////////
+
+FString FVisualGraphNode::DumpDot(const FVisualGraph* InGraph, int32 InIndendation) const
+{
+	const FString Indentation = DumpDotIndentation(InIndendation);
+	const FString StyleContent = DumpDotStyle(GetStyle());
+	const FString ShapeContent = DumpDotShape(GetShape());
+	const FString ColorContent = DumpDotColor(GetColor());
+
+	FString AttributeContent;
+	if(DisplayName.IsSet())
+	{
+		AttributeContent += FString::Printf(TEXT("label = \"%s\""), *DisplayName.GetValue().ToString());
+	}
+	if(Tooltip.IsSet())
+	{
+		if(!AttributeContent.IsEmpty())
+		{
+			AttributeContent += TEXT(", ");
+		}
+		AttributeContent += FString::Printf(TEXT("tooltip = \"%s\""), *Tooltip.GetValue());
+	}
+	if(!StyleContent.IsEmpty())
+	{
+		if(!AttributeContent.IsEmpty())
+		{
+			AttributeContent += TEXT(", ");
+		}
+		AttributeContent += FString::Printf(TEXT("style = %s"), *StyleContent);
+	}
+	if(!ShapeContent.IsEmpty())
+	{
+		if(!AttributeContent.IsEmpty())
+		{
+			AttributeContent += TEXT(", ");
+		}
+		AttributeContent += FString::Printf(TEXT("shape = %s"), *ShapeContent);
+	}
+	if(!ColorContent.IsEmpty())
+	{
+		if(!AttributeContent.IsEmpty())
+		{
+			AttributeContent += TEXT(", ");
+		}
+		AttributeContent += FString::Printf(TEXT("color = %s"), *ColorContent);
+
+		if(GetColor().IsSet())
+		{
+			if(!GetStyle().IsSet() || GetStyle().GetValue() == EVisualGraphStyle::Filled)
+			{
+				const FLinearColor CurrentColor = GetColor().GetValue();
+				const float Brightness = (0.299f * CurrentColor.R + 0.587f * CurrentColor.G + 0.114f * CurrentColor.B);
+				if(Brightness < 0.5f)
+				{
+					AttributeContent += TEXT(", fontcolor = white");
+				}
+			}
+		}
+	}
+
+	if(!AttributeContent.IsEmpty())
+	{
+		AttributeContent = FString::Printf(TEXT(" [ %s ]"), *AttributeContent);
+	}
+
+	return FString::Printf(TEXT("%s%s%s;\n"), *Indentation, *GetName().ToString(), *AttributeContent);
+}
+
diff --git a/Engine/Source/Developer/VisualGraphUtils/Private/VisualGraphObjectUtils.cpp b/Engine/Source/Developer/VisualGraphUtils/Private/VisualGraphObjectUtils.cpp
new file mode 100644
--- /dev/null
+++ b/Engine/Source/Developer/VisualGraphUtils/Private/VisualGraphObjectUtils.cpp
@@ -0,0 +1,805 @@
+// Copyright Epic Games, Inc. All Rights Reserved.
+
+#include "VisualGraphObjectUtils.h"
+#include "VisualGraphUtilsModule.h"
+#include "Elements/Framework/TypedElementListObjectUtil.h"
+#include "Serialization/ArchiveUObject.h"
+#include "GameFramework/Actor.h"
+#include "Components/ActorComponent.h"
+
+#if WITH_EDITOR
+
+#include "HAL/PlatformApplicationMisc.h"
+
+///////////////////////////////////////////////////////////////////////////////////
+/// Console Commands
+///////////////////////////////////////////////////////////////////////////////////
+
+FAutoConsoleCommandWithWorldAndArgs FCmdVisualGraphUtilsLogClassNames
+(
+	TEXT("VisualGraphUtils.Object.LogClassNames"),
+	TEXT("Logs all class path names given a partial name"),
+	FConsoleCommandWithWorldAndArgsDelegate::CreateLambda([](const TArray<FString>& InParams, UWorld* InWorld)
+	{
+		if(InParams.Num() == 0)
+		{
+			UE_LOG(LogVisualGraphUtils, Error, TEXT("Unsufficient parameters. Command usage:"));
+			UE_LOG(LogVisualGraphUtils, Error, TEXT("Example: VisualGraphUtils.Object.LogClassNames skeletal material"));
+			UE_LOG(LogVisualGraphUtils, Error, TEXT("Provide one or more search tokens for case insensitive search within all known (loaded) class names."));
+			return;
+		}
+
+		FString ClipboardContent;
+
+		for(const FString& SearchToken : InParams)
+		{
+			for (TObjectIterator<UClass> ClassIt; ClassIt; ++ClassIt)
+			{
+				if(ClassIt->GetName().Contains(SearchToken, ESearchCase::IgnoreCase))
+				{
+					const FString PathName = ClassIt->GetPathName();
+					ClipboardContent += PathName + TEXT("\n");
+					UE_LOG(LogVisualGraphUtils, Display, TEXT("Found Class %s"), *PathName);
+				}
+			}
+		}
+
+		FPlatformApplicationMisc::ClipboardCopy(*ClipboardContent);
+
+		UE_LOG(LogVisualGraphUtils, Display, TEXT("The content has also been copied to the clipboard."));
+	})
+);
+
+FAutoConsoleCommandWithWorldAndArgs FCmdVisualGraphUtilsLogInstancesOfClass
+(
+	TEXT("VisualGraphUtils.Object.LogInstancesOfClass"),
+	TEXT("Logs all instances of a given class"),
+	FConsoleCommandWithWorldAndArgsDelegate::CreateLambda([](const TArray<FString>& InParams, UWorld* InWorld)
+	{
+		if(InParams.Num() == 0)
+		{
+			UE_LOG(LogVisualGraphUtils, Error, TEXT("Unsufficient parameters. Command usage:"));
+			UE_LOG(LogVisualGraphUtils, Error, TEXT("Example: VisualGraphUtils.Object.LogInstancesOfClass /Script/Engine.EdGraphNode"));
+			UE_LOG(LogVisualGraphUtils, Error, TEXT("Provide one or more class path names"));
+			return;
+		}
+
+		FString ClipboardContent;
+		for(const FString& ObjectPathName : InParams)
+		{
+			if(UClass* Class = FindObject<UClass>(ANY_PACKAGE, *ObjectPathName, false))
+			{
+				ForEachObjectOfClass(Class, [&ClipboardContent](UObject* ObjectOfClass)
+				{
+					const FString PathName = ObjectOfClass->GetPathName();
+					ClipboardContent += PathName + TEXT("\n");
+					UE_LOG(LogVisualGraphUtils, Display, TEXT("Found Instance %s"), *PathName);
+				}, true, RF_NoFlags);
+			}
+			else
+			{
+				UE_LOG(LogVisualGraphUtils, Error, TEXT("Class with pathname '%s' not found."), *ObjectPathName);
+				return;
+			}
+		}
+
+		FPlatformApplicationMisc::ClipboardCopy(*ClipboardContent);
+
+		UE_LOG(LogVisualGraphUtils, Display, TEXT("The content has also been copied to the clipboard."));
+	})
+);
+
+FAutoConsoleCommandWithWorldAndArgs FCmdVisualGraphUtilsCollectReferences
+(
+	TEXT("VisualGraphUtils.Object.CollectReferences"),
+	TEXT("Traces all references of an object"),
+	FConsoleCommandWithWorldAndArgsDelegate::CreateLambda([](const TArray<FString>& InParams, UWorld* InWorld)
+	{
+		if(InParams.Num() == 0)
+		{
+			UE_LOG(LogVisualGraphUtils, Error, TEXT("Unsufficient parameters. Command usage:"));
+			UE_LOG(LogVisualGraphUtils, Error, TEXT("Example: VisualGraphUtils.Object.CollectReferences /Game/Animation/ControlRig/BasicControls_CtrlRig skip=/Script/RigVMDeveloper.RigVMNode skip=/Script/RigVMDeveloper.RigVMLink skip=/Script/Engine.EdGraph skip=/Script/Engine.EdGraphNode withinpackageonly"));
+			UE_LOG(LogVisualGraphUtils, Error, TEXT("Provide one or more path names for packages or objects to collect all references."));
+			UE_LOG(LogVisualGraphUtils, Error, TEXT("Optionally provide one or more path names for classes / outers to skip using the skip= token."));
+			UE_LOG(LogVisualGraphUtils, Error, TEXT("Use the the withinpackageonly token to limit the search within the source package(s) only."));
+			UE_LOG(LogVisualGraphUtils, Error, TEXT("Use the the skipchildren token to skip following objects within each outer (children)."));
+			UE_LOG(LogVisualGraphUtils, Error, TEXT("Use the the skipreferences token to skip following direct references."));
+			return;
+		}
+
+		bool bTraverseOuters = true;
+		bool bCollectReferences = true;
+		bool bWithinPackageOnly = false;
+
+		TArray<FString> ObjectPathNames;
+		TArray<FString> SkipPathNames;
+		for(const FString& InParam : InParams)
+		{
+			static const FString ObjectPathToken = TEXT("path"); 
+			static const FString SkipToken = TEXT("skip");
+			static const FString SkipChildrenToken = TEXT("skipchildren");
+			static const FString SkipReferencesToken = TEXT("skipreferences");
+			static const FString WithinPackageOnlyToken = TEXT("withinpackageonly");
+			FString Token = ObjectPathToken;
+			FString Content = InParam;
+
+			if(InParam.Contains(TEXT("=")))
+			{
+				const int32 Pos = InParam.Find(TEXT("="));
+				Token = InParam.Left(Pos);
+				Content = InParam.Mid(Pos+1);
+				Token.TrimStartAndEndInline();
+				Token.ToLowerInline();
+				Content.TrimStartAndEndInline();
+			}
+			else if(InParam.ToLower() == SkipChildrenToken)
+			{
+				bTraverseOuters = false;
+				continue;
+			}
+			else if(InParam.ToLower() == SkipReferencesToken)
+			{
+				bCollectReferences = false;
+				continue;
+			}
+			else if(InParam.ToLower() == WithinPackageOnlyToken)
+			{
+				bWithinPackageOnly = true;
+				continue;
+			}
+
+			if(Token == ObjectPathToken)
+			{
+				ObjectPathNames.Add(Content);
+			}
+			else if(Token == SkipToken)
+			{
+				SkipPathNames.Add(Content);
+			}
+		}
+
+		TArray<UObject*> Objects;
+		TArray<UObject*> OutersToUse;
+		for(const FString& ObjectPathName : ObjectPathNames)
+		{
+			if(UObject* Object = FindObject<UObject>(ANY_PACKAGE, *ObjectPathName, false))
+			{
+				Objects.Add(Object);
+
+				if(bWithinPackageOnly)
+				{
+					OutersToUse.AddUnique(Object->GetOutermost());
+				}
+			}
+			else
+			{
+				UE_LOG(LogVisualGraphUtils, Error, TEXT("Object with pathname '%s' not found."), *ObjectPathName);
+				return;
+			}
+		}
+
+		TArray<UClass*> ClassesToSkip;
+		TArray<UObject*> OutersToSkip;
+		for(const FString& SkipPathName : SkipPathNames)
+		{
+			if(UObject* ObjectToSkip = FindObject<UObject>(ANY_PACKAGE, *SkipPathName, false))
+			{
+				if(UClass* ClassToSkip = Cast<UClass>(ObjectToSkip))
+				{
+					ClassesToSkip.Add(ClassToSkip);
+				}
+				else
+				{
+					OutersToSkip.Add(ObjectToSkip);
+				}
+			}
+			else
+			{
+				UE_LOG(LogVisualGraphUtils, Error, TEXT("Object with pathname '%s' not found."), *SkipPathName);
+				return;
+			}
+		}
+
+		const FString DotGraphContent = FVisualGraphObjectUtils::TraverseUObjectReferences(
+			Objects, ClassesToSkip, OutersToSkip, OutersToUse,
+			bTraverseOuters, bCollectReferences).DumpDot();
+		FPlatformApplicationMisc::ClipboardCopy(*DotGraphContent);
+
+		UE_LOG(LogVisualGraphUtils, Display, TEXT("The content has been copied to the clipboard."));
+	})
+);
+
+FAutoConsoleCommandWithWorldAndArgs FCmdVisualGraphUtilsCollectTickables
+(
+	TEXT("VisualGraphUtils.Object.CollectTickables"),
+	TEXT("Traces all tickables of an object"),
+	FConsoleCommandWithWorldAndArgsDelegate::CreateLambda([](const TArray<FString>& InParams, UWorld* InWorld)
+	{
+		if(InParams.Num() == 0)
+		{
+			UE_LOG(LogVisualGraphUtils, Error, TEXT("Unsufficient parameters. Command usage:"));
+			UE_LOG(LogVisualGraphUtils, Error, TEXT("Example: VisualGraphUtils.Object.CollectTickables /Game/Animation/ControlRig/BasicControls_CtrlRig skip=/Script/RigVMDeveloper.RigVMNode skip=/Script/RigVMDeveloper.RigVMLink skip=/Script/Engine.EdGraph skip=/Script/Engine.EdGraphNode withinpackageonly"));
+			UE_LOG(LogVisualGraphUtils, Error, TEXT("Provide one or more path names for packages or objects to collect all tickables."));
+			return;
+		}
+
+		TArray<FString> ObjectPathNames;
+		for(const FString& InParam : InParams)
+		{
+			static const FString ObjectPathToken = TEXT("path"); 
+			FString Token = ObjectPathToken;
+			FString Content = InParam;
+
+			if(InParam.Contains(TEXT("=")))
+			{
+				const int32 Pos = InParam.Find(TEXT("="));
+				Token = InParam.Left(Pos);
+				Content = InParam.Mid(Pos+1);
+				Token.TrimStartAndEndInline();
+				Token.ToLowerInline();
+				Content.TrimStartAndEndInline();
+			}
+
+			if(Token == ObjectPathToken)
+			{
+				ObjectPathNames.Add(Content);
+			}
+		}
+
+		TArray<UObject*> Objects;
+		TArray<UObject*> OutersToUse;
+		for(const FString& ObjectPathName : ObjectPathNames)
+		{
+			if(UObject* Object = FindObject<UObject>(ANY_PACKAGE, *ObjectPathName, false))
+			{
+				Objects.Add(Object);
+			}
+			else
+			{
+				UE_LOG(LogVisualGraphUtils, Error, TEXT("Object with pathname '%s' not found."), *ObjectPathName);
+				return;
+			}
+		}
+
+		const FString DotGraphContent = FVisualGraphObjectUtils::TraverseTickOrder(
+			Objects).DumpDot();
+		FPlatformApplicationMisc::ClipboardCopy(*DotGraphContent);
+
+		UE_LOG(LogVisualGraphUtils, Display, TEXT("The content has been copied to the clipboard."));
+	})
+);
+
+#endif
+
+///////////////////////////////////////////////////////////////////////////////////
+/// FVisualGraphObjectUtilsReferenceCollector
+///////////////////////////////////////////////////////////////////////////////////
+
+class FVisualGraphObjectUtilsReferenceCollector : public FArchiveUObject
+{
+public:
+	
+	enum EReferenceKind
+	{
+		EReferenceKind_Direct,
+		EReferenceKind_Outer
+	};
+
+	struct FReference
+	{
+		EReferenceKind Kind;
+		UObject* Object;
+
+		bool operator ==(const FReference& InReference) const
+		{
+			return InReference.Object == Object;
+		}
+
+		friend uint32 GetTypeHash(const FReference& InReference)
+		{
+			return GetTypeHash(InReference.Object);
+		}
+	};
+
+private:
+
+	// I/O function.  Called when an object reference is encountered.
+	FArchive& operator<<( UObject*& Obj ) override
+	{
+		if( Obj )
+		{
+			FoundReference( Obj, EReferenceKind_Direct );
+		}
+		return *this;
+	}
+
+	virtual FArchive& operator<< (struct FSoftObjectPtr& Value) override
+	{
+		if ( Value.Get() )
+		{
+			FoundReference(Value.Get(), EReferenceKind_Direct);
+
+			if(bCollectReferencesBySerialize)
+			{
+				Value.Get()->Serialize( *this );
+			}
+		}
+		return *this;
+	}
+	
+	virtual FArchive& operator<< (struct FSoftObjectPath& Value) override
+	{
+		if ( Value.ResolveObject() )
+		{
+			FoundReference(Value.ResolveObject(), EReferenceKind_Direct);
+
+			if(bCollectReferencesBySerialize)
+			{
+				Value.ResolveObject()->Serialize( *this );
+			}
+		}
+		return *this;
+	}
+
+	void FoundReference( UObject* Object, EReferenceKind Kind )
+	{
+		FReference Target;
+		Target.Kind = Kind;
+		Target.Object = Object;
+		const TPair<UObject*, FReference> Reference(CurrentObject, Target);
+
+		if(!ClassesToSkip.IsEmpty())
+		{
+			for(const UClass* ClassToSkip : ClassesToSkip)
+			{
+				if(Object->GetClass()->IsChildOf(ClassToSkip))
+				{
+					return;
+				}
+			}
+		}
+
+		if(!OutersToSkip.IsEmpty())
+		{
+			for(const UObject* OuterToSkip : OutersToSkip)
+			{
+				if(Object->IsInOuter(OuterToSkip))
+				{
+					return;
+				}
+			}
+		}
+
+		if(!OutersToUse.IsEmpty())
+		{
+			for(const UObject* OuterToUse : OutersToUse)
+			{
+				if(!Object->IsInOuter(OuterToUse))
+				{
+					return;
+				}
+			}
+		}
+
+		if ( RootSet.Find(Object) == nullptr )
+		{
+			if(bRecursive)
+			{
+				RootSetArray.Add( Object );
+			}
+			RootSet.Add(Object);
+			References.Add(Reference);
+		}
+		else
+		{
+			References.FindOrAdd(Reference);
+		}
+	}
+
+	UObject* CurrentObject;
+	TSet<TPair<UObject*, FReference>>& References;
+	TArray<UObject*> RootSetArray;
+	TSet<UObject*> RootSet;
+	const TArray<UClass*>& ClassesToSkip;
+	const TArray<UObject*>& OutersToSkip;
+	const TArray<UObject*>& OutersToUse;
+	bool bTraverseObjectsInOuter;
+	bool bCollectReferencesBySerialize;
+	bool bRecursive;
+
+public:
+
+	FVisualGraphObjectUtilsReferenceCollector(
+		TSet<UObject*> InRootSet,
+		TSet<TPair<UObject*, FReference>>& InReferences,
+		const TArray<UClass*>& InClassesToSkip,
+		const TArray<UObject*>& InOutersToSkip,
+		const TArray<UObject*>& InOutersToUse,
+		bool InTraverseObjectsInOuter,
+		bool InCollectReferencesBySerialize,
+		bool InRecursive
+		)
+		: CurrentObject(nullptr)
+		, References(InReferences)
+		, RootSet(InRootSet)
+		, ClassesToSkip(InClassesToSkip)
+		, OutersToSkip(InOutersToSkip)
+		, OutersToUse(InOutersToUse)
+		, bTraverseObjectsInOuter(InTraverseObjectsInOuter)
+		, bCollectReferencesBySerialize(InCollectReferencesBySerialize)
+		, bRecursive(InRecursive)
+	{
+		ArIsObjectReferenceCollector = true;
+		this->SetIsSaving(true);
+
+		for ( UObject* Object : RootSet )
+		{
+			RootSetArray.Add( Object );
+		}
+		
+		// loop through all the objects in the root set and serialize them
+		for ( int RootIndex = 0; RootIndex < RootSetArray.Num(); ++RootIndex )
+		{
+			UObject* SourceObject = RootSetArray[RootIndex];
+
+			// quick sanity check
+			check(SourceObject);
+			check(SourceObject->IsValidLowLevel());
+
+			TGuardValue<UObject*> Guard(CurrentObject, SourceObject);
+
+			if(bCollectReferencesBySerialize)
+			{
+				SourceObject->Serialize( *this );
+			}
+
+			if(bTraverseObjectsInOuter)
+			{
+				ForEachObjectWithOuter(SourceObject, [this](UObject* Inner)
+				{
+					FoundReference(Inner, EReferenceKind_Outer);
+				}, false);
+			}
+		}
+
+	}
+
+	virtual FString GetArchiveName() const override { return TEXT("FVisualGraphObjectUtilsReferenceCollector"); }
+};
+
+///////////////////////////////////////////////////////////////////////////////////
+/// FVisualGraphObjectUtils
+///////////////////////////////////////////////////////////////////////////////////
+
+FVisualGraph FVisualGraphObjectUtils::TraverseUObjectReferences(
+	const TArray<UObject*>& InObjects,
+	const TArray<UClass*>& InClassesToSkip,
+	const TArray<UObject*>& InOutersToSkip,
+	const TArray<UObject*>& InOutersToUse,
+	bool bTraverseObjectsInOuter,
+	bool bCollectReferencesBySerialize,
+	bool bRecursive)
+{
+	typedef FVisualGraphObjectUtilsReferenceCollector::FReference FReference;
+
+	FVisualGraph Graph(TEXT("References"));
+	if(InObjects.IsEmpty())
+	{
+		return Graph;
+	}
+
+	TSet<UObject*> RootObjects;
+	for(UObject* InObject : InObjects)
+	{
+		RootObjects.Add(InObject);
+	}
+
+	struct Local
+	{
+		static int32 TraverseUObjectReferences(UObject* InObject, FVisualGraph& OutGraph)
+		{
+			const FName GuidName = *FString::Printf(TEXT("node_%d"), (int32)InObject->GetUniqueID());
+			int32 NodeIndex = OutGraph.FindNode(GuidName);
+			if(NodeIndex != INDEX_NONE)
+			{
+				return NodeIndex;
+			}
+
+			const FString PathName = InObject->GetPathName();
+			if(PathName.StartsWith(TEXT("/Script/")))
+			{
+				return INDEX_NONE;
+			}
+			if(PathName.Contains(TEXT("TRASH_"), ESearchCase::CaseSensitive))
+			{
+				return INDEX_NONE;
+			}
+
+			NodeIndex = OutGraph.AddNode(GuidName, InObject->GetFName());
+
+			OutGraph.GetNodes()[NodeIndex].SetTooltip(PathName);
+
+			return NodeIndex;
+		}
+	};
+
+	TSet<TPair<UObject*, FReference>> References;
+	FVisualGraphObjectUtilsReferenceCollector Collector(
+		RootObjects,
+		References,
+		InClassesToSkip,
+		InOutersToSkip,
+		InOutersToUse,
+		bTraverseObjectsInOuter,
+		bCollectReferencesBySerialize,
+		bRecursive);
+	
+	for(const TPair<UObject*, FReference>& Reference : References)
+	{
+		UObject* Source = Reference.Key;
+		const FReference& Target = Reference.Value;
+		
+		const int32 SourceNodeIndex = Local::TraverseUObjectReferences(Source, Graph);
+		const int32 TargetNodeIndex = Local::TraverseUObjectReferences(Target.Object, Graph);
+
+		if(SourceNodeIndex != INDEX_NONE && TargetNodeIndex != INDEX_NONE)
+		{
+			TOptional<EVisualGraphStyle> Style;
+			
+			if(Target.Kind == FVisualGraphObjectUtilsReferenceCollector::EReferenceKind_Outer)
+			{
+				Style = EVisualGraphStyle::Dashed;
+			}
+			
+			Graph.AddEdge(
+				SourceNodeIndex,
+				TargetNodeIndex,
+				EVisualGraphEdgeDirection::TargetToSource,
+				NAME_None,
+				TOptional<FName>(),
+				TOptional<FLinearColor>(),
+				Style);
+		}
+	}
+	
+	return Graph;
+}
+
+FVisualGraph FVisualGraphObjectUtils::TraverseTickOrder(const TArray<UObject*>& InObjects)
+{
+	typedef FVisualGraphObjectUtilsReferenceCollector::FReference FReference;
+
+	FVisualGraph Graph(TEXT("References"));
+	if(InObjects.IsEmpty())
+	{
+		return Graph;
+	}
+
+	TSet<UObject*> RootObjects;
+	for(UObject* InObject : InObjects)
+	{
+		RootObjects.Add(InObject);
+	}
+
+	struct Local
+	{
+		static bool HasTickPrerequisite(
+			UObject* InTargetObject,
+			UObject* InSourceObject)
+		{
+			// we have found a potential edge - do we also have a tick order relationship?
+			if(const AActor* TargetActor = Cast<AActor>(InTargetObject))
+			{
+				const TArray<struct FTickPrerequisite>& Prerequisites =
+					TargetActor->PrimaryActorTick.GetPrerequisites();
+
+				for(const struct FTickPrerequisite& Prerequisite : Prerequisites)
+				{
+					if(UObject* TickObject = Prerequisite.PrerequisiteObject.Get())
+					{
+						if(TickObject == InSourceObject)
+						{
+							return true;
+						}
+					}
+				}
+
+				if(const UActorComponent* SourceComponent = Cast<UActorComponent>(InSourceObject))
+				{
+					if(SourceComponent->GetOwner() == TargetActor)
+					{
+						return true;
+					}
+				}
+			}
+
+			if(const UActorComponent* TargetComponent = Cast<UActorComponent>(InTargetObject))
+			{
+				const TArray<struct FTickPrerequisite>& Prerequisites =
+					TargetComponent->PrimaryComponentTick.GetPrerequisites();
+
+				for(const struct FTickPrerequisite& Prerequisite : Prerequisites)
+				{
+					if(UObject* TickObject = Prerequisite.PrerequisiteObject.Get())
+					{
+						if(TickObject == InSourceObject)
+						{
+							return true;
+						}
+					}
+				}
+				
+				if(const UActorComponent* SourceComponent = Cast<UActorComponent>(InSourceObject))
+				{
+					if(TargetComponent->IsInOuter(SourceComponent))
+					{
+						return true;
+					}
+				}
+			}
+
+			return false;
+		}
+
+		static void TraverseTickOrderEdges(
+			const UObject* InTargetObject,
+			const int32 InTargetNodeIndex,
+			UObject* InTraversingObject,
+			FVisualGraph& OutGraph,
+			const TMap<UObject*, TArray<FReference>>& InTargetToSource,
+			TMap<TPair<UObject*, UObject*>, bool>& OutVisitedEdges,
+			TMap<UObject*, int32>& OutFarthestDistance,
+			int32 InDistance
+			)
+		{
+			check(InTargetNodeIndex != INDEX_NONE);
+			
+			const TArray<FReference>& Sources = InTargetToSource.FindChecked(InTargetObject);
+			for(const FReference& Source : Sources)
+			{
+				TPair<UObject*, UObject*> Edge(Source.Object, InTraversingObject);
+				if(OutVisitedEdges.Contains(Edge))
+				{
+					return;
+				}
+				OutVisitedEdges.Add(Edge, true);
+
+				if(Source.Object == InTraversingObject)
+				{
+					return;
+				}
+
+				const FName GuidName = *FString::Printf(TEXT("node_%d"), (int32)Source.Object->GetUniqueID());
+				const int32 SourceNodeIndex = OutGraph.FindNode(GuidName);
+				if(SourceNodeIndex != INDEX_NONE)
+				{
+					
+					TOptional<EVisualGraphStyle> Style;
+
+					// we have found a potential edge - do we also have a tick order relationship?
+					if(!HasTickPrerequisite(InTraversingObject, Source.Object) &&
+						!HasTickPrerequisite(Source.Object, InTraversingObject))
+					{
+						Style = EVisualGraphStyle::Dotted;
+					}
+			
+					OutGraph.AddEdge(
+						SourceNodeIndex,
+						InTargetNodeIndex,
+						EVisualGraphEdgeDirection::TargetToSource,
+						NAME_None,
+						TOptional<FName>(),
+						TOptional<FLinearColor>(),
+						Style);
+				}
+
+				TraverseTickOrderEdges(Source.Object, InTargetNodeIndex, InTraversingObject, OutGraph, InTargetToSource, OutVisitedEdges, OutFarthestDistance, InDistance + 1);
+			}
+		}
+
+		static int32 TraverseTickOrder(
+			UObject* InObject,
+			FVisualGraph& OutGraph,
+			const TMap<UObject*, TArray<FReference>>& InTargetToSource,
+			TMap<UObject*, bool>& OutVisitedObjects,
+			TMap<TPair<UObject*, UObject*>, bool>& OutVisitedEdges,
+			TMap<UObject*, int32>& OutFarthestDistance
+			)
+		{
+			const FName GuidName = *FString::Printf(TEXT("node_%d"), (int32)InObject->GetUniqueID());
+			int32 NodeIndex = OutGraph.FindNode(GuidName);
+			if(NodeIndex != INDEX_NONE)
+			{
+				return NodeIndex;
+			}
+
+			if(OutVisitedObjects.Contains(InObject))
+			{
+				return INDEX_NONE;
+			}
+			OutVisitedObjects.Add(InObject, true);
+
+			const FString PathName = InObject->GetPathName();
+			if(PathName.StartsWith(TEXT("/Script/")))
+			{
+				return INDEX_NONE;
+			}
+			if(PathName.Contains(TEXT("TRASH_"), ESearchCase::CaseSensitive))
+			{
+				return INDEX_NONE;
+			}
+
+			bool bHasTick = false;
+			if(const AActor* Actor = Cast<AActor>(InObject))
+			{
+				bHasTick = Actor->CanEverTick();
+			}
+			else if(const UActorComponent* ActorComponent = Cast<UActorComponent>(InObject))
+			{
+				bHasTick = ActorComponent->IsComponentTickEnabled();
+			}
+
+			if(bHasTick)
+			{
+				NodeIndex = OutGraph.AddNode(GuidName, InObject->GetFName());
+			}
+
+			// visit all targets. we do this to be sure that sources are processed prior to targets
+			const TArray<FReference>& Sources = InTargetToSource.FindChecked(InObject);
+			for(const FReference& Source : Sources)
+			{
+				TraverseTickOrder(Source.Object, OutGraph, InTargetToSource, OutVisitedObjects, OutVisitedEdges, OutFarthestDistance);
+			}
+
+			// for myself - walk back up and see if I am going to hit anything
+			if(NodeIndex != INDEX_NONE)
+			{
+				OutGraph.GetNodes()[NodeIndex].SetTooltip(PathName);
+				TraverseTickOrderEdges(InObject, NodeIndex, InObject, OutGraph, InTargetToSource, OutVisitedEdges, OutFarthestDistance, 0);
+			}
+
+			return NodeIndex;
+		}
+	};
+
+	TSet<TPair<UObject*, FReference>> References;
+	FVisualGraphObjectUtilsReferenceCollector Collector(
+		RootObjects,
+		References,
+		TArray<UClass*>(),
+		TArray<UObject*>(),
+		TArray<UObject*>(),
+		true,
+		true,
+		true);
+
+	TMap<UObject*, TArray<FReference>> TargetToSource;
+
+	for(const TPair<UObject*, FReference>& Reference : References)
+	{
+		TArray<FReference>& Sources = TargetToSource.FindOrAdd(Reference.Value.Object);
+		FReference ReverseReference = Reference.Value;
+		ReverseReference.Object = Reference.Key;
+		Sources.Add(ReverseReference);
+		TargetToSource.FindOrAdd(Reference.Key);
+	}
+
+	TMap<UObject*, bool> VisitedObjects;
+	TMap<TPair<UObject*, UObject*>, bool> VisitedEdges;
+	TMap<UObject*, int32> ObjectDistance;
+	for(const TPair<UObject*, FReference>& Reference : References)
+	{
+		Local::TraverseTickOrder(Reference.Key, Graph, TargetToSource, VisitedObjects, VisitedEdges, ObjectDistance);
+	}
+
+	Graph.TransitiveReduction([](FVisualGraphEdge& Edge) -> bool
+	{
+		Edge.SetColor(FLinearColor::Gray);
+		return true /* keep edge */;
+	});
+	
+	return Graph;
+}
+
+
+
diff --git a/Engine/Source/Developer/VisualGraphUtils/Private/VisualGraphUtilsModule.cpp b/Engine/Source/Developer/VisualGraphUtils/Private/VisualGraphUtilsModule.cpp
new file mode 100644
--- /dev/null
+++ b/Engine/Source/Developer/VisualGraphUtils/Private/VisualGraphUtilsModule.cpp
@@ -0,0 +1,8 @@
+// Copyright Epic Games, Inc. All Rights Reserved.
+
+#include "VisualGraphUtilsModule.h"
+#include "Modules/ModuleManager.h"
+
+IMPLEMENT_MODULE(FDefaultModuleImpl, VisualGraphUtils);
+
+DEFINE_LOG_CATEGORY(LogVisualGraphUtils);
diff --git a/Engine/Source/Developer/VisualGraphUtils/Public/VisualGraph.h b/Engine/Source/Developer/VisualGraphUtils/Public/VisualGraph.h
new file mode 100644
--- /dev/null
+++ b/Engine/Source/Developer/VisualGraphUtils/Public/VisualGraph.h
@@ -0,0 +1,135 @@
+﻿// Copyright Epic Games, Inc. All Rights Reserved.
+
+#pragma once
+
+#include "VisualGraphElement.h"
+#include "VisualGraphNode.h"
+#include "VisualGraphEdge.h"
+
+class FVisualGraph;
+
+class VISUALGRAPHUTILS_API FVisualGraphSubGraph : public FVisualGraphElement
+{
+public:
+
+	FVisualGraphSubGraph()
+	: FVisualGraphElement()
+	, ParentGraphIndex(INDEX_NONE)
+	{}
+
+	virtual ~FVisualGraphSubGraph() override {}
+
+	int32 GetParentGraphIndex() const { return ParentGraphIndex; }
+	const TArray<int32>& GetNodes() const { return Nodes; }
+
+protected:
+
+	virtual FString DumpDot(const FVisualGraph* InGraph, int32 InIndendation) const override;
+
+	int32 ParentGraphIndex;
+	TArray<int32> Nodes;
+
+	friend class FVisualGraph;
+};
+
+class VISUALGRAPHUTILS_API FVisualGraph : public FVisualGraphElement
+{
+public:
+
+	FVisualGraph()
+	: FVisualGraphElement()
+	{}
+
+	virtual ~FVisualGraph() override {}
+
+	FVisualGraph(const FName& InName, const FName& InDisplayName = NAME_None);
+
+	const TArray<FVisualGraphNode>& GetNodes() const { return Nodes; }
+	const TArray<FVisualGraphEdge>& GetEdges() const { return Edges; }
+	const TArray<FVisualGraphSubGraph>& GetSubGraphs() const { return SubGraphs; }
+
+	int32 AddNode(
+		const FName& InName,
+		TOptional<FName> InDisplayName = TOptional<FName>(),
+		TOptional<FLinearColor> InColor = TOptional<FLinearColor>(),
+		TOptional<EVisualGraphShape> InShape = TOptional<EVisualGraphShape>(), 
+		TOptional<EVisualGraphStyle> InStyle = TOptional<EVisualGraphStyle>());
+
+	int32 AddEdge(
+		int32 InSourceNode,
+		int32 InTargetNode,
+		EVisualGraphEdgeDirection InDirection,
+		const FName& InName = NAME_None,
+		TOptional<FName> InDisplayName = TOptional<FName>(),
+		TOptional<FLinearColor> InColor = TOptional<FLinearColor>(),
+		TOptional<EVisualGraphStyle> InStyle = TOptional<EVisualGraphStyle>());
+
+	int32 AddSubGraph(
+		const FName& InName,
+		TOptional<FName> InDisplayName = TOptional<FName>(),
+		int32 InParentGraphIndex = INDEX_NONE,
+		TOptional<FLinearColor> InColor = TOptional<FLinearColor>(),
+		TOptional<EVisualGraphStyle> InStyle = TOptional<EVisualGraphStyle>(),
+		const TArray<int32> InNodes = TArray<int32>());
+
+	int32 FindNode(const FName& InName) const;
+	int32 FindEdge(const FName& InName) const;
+	int32 FindSubGraph(const FName& InName) const;
+
+	bool AddNodeToSubGraph(int32 InNodeIndex, int32 InSubGraphIndex);
+	bool RemoveNodeFromSubGraph(int32 InNodeIndex);
+
+	void TransitiveReduction(TFunction<bool(FVisualGraphEdge&)> KeepEdgeFunction);
+
+	FORCEINLINE FString DumpDot() const
+	{
+		return DumpDot(this, 0);
+	}
+
+protected:
+
+	virtual FString DumpDot(const FVisualGraph* InGraph, int32 InIndendation) const override;
+
+	template<typename T>
+	FORCEINLINE static void RefreshNameMap(const TArray<T>& InElements, TMap<FName, int32>& OutMap)
+	{
+		OutMap.Reset();
+		for(const T& Element: InElements)
+		{
+			OutMap.Add(Element.Name, Element.Index);
+		}
+	}
+
+	template<typename T>
+	FORCEINLINE static void RefreshNameMapIfNeeded(const TArray<T>& InElements, TMap<FName, int32>& OutMap)
+	{
+		if(OutMap.Num() == InElements.Num())
+		{
+			return;
+		}
+		RefreshNameMap(InElements, OutMap);
+	}
+
+	template<typename T>
+	FORCEINLINE static int32 AddElement(const T& InElement, TArray<T>& OutElements, TMap<FName, int32>& OutMap)
+	{
+		const int32 AddedIndex = OutElements.Add(InElement);
+		if(OutElements.IsValidIndex(AddedIndex))
+		{
+			OutElements[AddedIndex].Index = AddedIndex;
+			OutMap.Add(InElement.Name, AddedIndex);
+		}
+		return AddedIndex;
+	}
+
+	static bool IsNameAvailable(const FName& InName, const TMap<FName, int32>& InMap);
+	static FName GetUniqueName(const FName& InName, const TMap<FName, int32>& InMap);
+
+	TArray<FVisualGraphNode> Nodes;
+	TArray<FVisualGraphEdge> Edges;
+	TArray<FVisualGraphSubGraph> SubGraphs;
+
+	TMap<FName,int32> NodeNameMap;
+	TMap<FName,int32> EdgeNameMap;
+	TMap<FName,int32> SubGraphNameMap;
+};
diff --git a/Engine/Source/Developer/VisualGraphUtils/Public/VisualGraphEdge.h b/Engine/Source/Developer/VisualGraphUtils/Public/VisualGraphEdge.h
new file mode 100644
--- /dev/null
+++ b/Engine/Source/Developer/VisualGraphUtils/Public/VisualGraphEdge.h
@@ -0,0 +1,42 @@
+﻿// Copyright Epic Games, Inc. All Rights Reserved.
+
+#pragma once
+
+#include "VisualGraphElement.h"
+#include "VisualGraphEdge.generated.h"
+
+UENUM()
+enum class EVisualGraphEdgeDirection : uint8
+{
+	SourceToTarget,
+	TargetToSource,
+	BothWays
+};
+
+class VISUALGRAPHUTILS_API FVisualGraphEdge : public FVisualGraphElement
+{
+public:
+
+	FVisualGraphEdge()
+	: FVisualGraphElement()
+	, Direction(EVisualGraphEdgeDirection::SourceToTarget)
+	{}
+
+	virtual ~FVisualGraphEdge() override {}
+
+	int32 GetSourceNode() const { return SourceNode; }
+	int32 GetTargetNode() const { return TargetNode; }
+	EVisualGraphEdgeDirection GetDirection() const { return Direction; }
+
+protected:
+
+	virtual FString DumpDot(const FVisualGraph* InGraph, int32 InIndendation) const override;
+
+	int32 SourceNode;
+	int32 TargetNode;
+	EVisualGraphEdgeDirection Direction;
+
+	friend class FVisualGraph;
+	friend class FVisualGraphSubGraph;
+};
+
diff --git a/Engine/Source/Developer/VisualGraphUtils/Public/VisualGraphElement.h b/Engine/Source/Developer/VisualGraphUtils/Public/VisualGraphElement.h
new file mode 100644
--- /dev/null
+++ b/Engine/Source/Developer/VisualGraphUtils/Public/VisualGraphElement.h
@@ -0,0 +1,77 @@
+﻿// Copyright Epic Games, Inc. All Rights Reserved.
+
+#pragma once
+
+#include "CoreMinimal.h"
+#include "VisualGraphElement.generated.h"
+
+class FVisualGraph;
+
+UENUM()
+enum class EVisualGraphShape : uint8
+{
+	Box,
+	Polygon,
+	Ellipse,
+	Circle,
+	Triangle,
+	PlainText,
+	Diamond,
+	Parallelogram,
+	House
+};
+
+UENUM()
+enum class EVisualGraphStyle : uint8
+{
+	Filled,
+	Diagonals,
+	Rounded,
+	Dashed,
+	Dotted,
+	Solid,
+	Bold
+};
+
+class VISUALGRAPHUTILS_API FVisualGraphElement
+{
+public:
+
+	FVisualGraphElement()
+	: Name(NAME_None)
+	, DisplayName()
+	, Index(INDEX_NONE)
+	, Color()
+	, Style()
+	{}
+
+	virtual ~FVisualGraphElement() {}
+
+	FName GetName() const { return Name; }
+	FName GetDisplayName() const { return DisplayName.IsSet() ? DisplayName.GetValue() : Name; }
+	TOptional<FString> GetTooltip() const { return Tooltip; }
+	void SetTooltip(const FString& InTooltip) const { Tooltip = InTooltip; }
+	int32 GetIndex() const { return Index; }
+	TOptional<FLinearColor> GetColor() const { return Color; }
+	void SetColor(FLinearColor InValue) { Color = InValue; }
+	TOptional<EVisualGraphStyle> GetStyle() const { return Style; }
+	void SetStyle(EVisualGraphStyle InValue) { Style = InValue; }
+
+protected:
+
+	virtual FString DumpDot(const FVisualGraph* InGraph, int32 InIndendation) const = 0;
+	static FString DumpDotIndentation(int32 InIndentation);
+	static FString DumpDotColor(const TOptional<FLinearColor>& InColor);
+	static FString DumpDotShape(const TOptional<EVisualGraphShape>& InShape);
+	static FString DumpDotStyle(const TOptional<EVisualGraphStyle>& InStyle);
+
+	FName Name;
+	TOptional<FName> DisplayName;
+	mutable TOptional<FString> Tooltip;
+	int32 Index;
+	TOptional<FLinearColor> Color;
+	TOptional<EVisualGraphStyle> Style;
+	
+	friend class FVisualGraph;
+};
+
diff --git a/Engine/Source/Developer/VisualGraphUtils/Public/VisualGraphNode.h b/Engine/Source/Developer/VisualGraphUtils/Public/VisualGraphNode.h
new file mode 100644
--- /dev/null
+++ b/Engine/Source/Developer/VisualGraphUtils/Public/VisualGraphNode.h
@@ -0,0 +1,31 @@
+﻿// Copyright Epic Games, Inc. All Rights Reserved.
+
+#pragma once
+
+#include "VisualGraphElement.h"
+
+class VISUALGRAPHUTILS_API FVisualGraphNode : public FVisualGraphElement
+{
+public:
+
+	FVisualGraphNode()
+	: FVisualGraphElement()
+	, Shape()
+	, SubGraphIndex(INDEX_NONE)
+	{}
+
+	virtual ~FVisualGraphNode() override {}
+
+	TOptional<EVisualGraphShape> GetShape() const { return Shape; }
+	void SetShape(EVisualGraphShape InValue) { Shape = InValue; }
+
+protected:
+
+	virtual FString DumpDot(const FVisualGraph* InGraph, int32 InIndendation) const override;
+
+	TOptional<EVisualGraphShape> Shape;
+	int32 SubGraphIndex;
+
+	friend class FVisualGraph;
+	friend class FVisualGraphSubGraph;
+};
diff --git a/Engine/Source/Developer/VisualGraphUtils/Public/VisualGraphObjectUtils.h b/Engine/Source/Developer/VisualGraphUtils/Public/VisualGraphObjectUtils.h
new file mode 100644
--- /dev/null
+++ b/Engine/Source/Developer/VisualGraphUtils/Public/VisualGraphObjectUtils.h
@@ -0,0 +1,25 @@
+﻿// Copyright Epic Games, Inc. All Rights Reserved.
+
+#pragma once
+
+#include "VisualGraph.h"
+
+class VISUALGRAPHUTILS_API FVisualGraphObjectUtils
+{
+public:
+
+	static FVisualGraph TraverseUObjectReferences(
+		const TArray<UObject*>& InObjects,
+		const TArray<UClass*>& InClassesToSkip = TArray<UClass*>(),
+		const TArray<UObject*>& InOutersToSkip = TArray<UObject*>(),
+		const TArray<UObject*>& InOutersToUse = TArray<UObject*>(),
+		bool bTraverseObjectsInOuter = true,
+		bool bCollectReferencesBySerialize = true,
+		bool bRecursive = true
+		);
+
+	static FVisualGraph TraverseTickOrder(
+		const TArray<UObject*>& InObjects
+		);
+};
+
diff --git a/Engine/Source/Developer/VisualGraphUtils/Public/VisualGraphUtils.h b/Engine/Source/Developer/VisualGraphUtils/Public/VisualGraphUtils.h
new file mode 100644
--- /dev/null
+++ b/Engine/Source/Developer/VisualGraphUtils/Public/VisualGraphUtils.h
@@ -0,0 +1,7 @@
+﻿// Copyright Epic Games, Inc. All Rights Reserved.
+
+#pragma once
+
+#include "VisualGraphNode.h"
+#include "VisualGraphEdge.h"
+#include "VisualGraph.h"
diff --git a/Engine/Source/Developer/VisualGraphUtils/Public/VisualGraphUtilsModule.h b/Engine/Source/Developer/VisualGraphUtils/Public/VisualGraphUtilsModule.h
new file mode 100644
--- /dev/null
+++ b/Engine/Source/Developer/VisualGraphUtils/Public/VisualGraphUtilsModule.h
@@ -0,0 +1,8 @@
+// Copyright Epic Games, Inc. All Rights Reserved.
+
+#pragma once
+
+#include "CoreMinimal.h"
+
+DECLARE_LOG_CATEGORY_EXTERN(LogVisualGraphUtils, Log, All);
+
diff --git a/Engine/Source/Developer/VisualGraphUtils/VisualGraphUtils.Build.cs b/Engine/Source/Developer/VisualGraphUtils/VisualGraphUtils.Build.cs
new file mode 100644
--- /dev/null
+++ b/Engine/Source/Developer/VisualGraphUtils/VisualGraphUtils.Build.cs
@@ -0,0 +1,27 @@
+// Copyright Epic Games, Inc. All Rights Reserved.
+
+using UnrealBuildTool;
+
+public class VisualGraphUtils : ModuleRules
+{
+    public VisualGraphUtils(ReadOnlyTargetRules Target) : base(Target)
+    {
+        PublicDependencyModuleNames.AddRange(
+            new string[] {
+                "Core",
+                "CoreUObject",
+                "Engine",
+            }
+        );
+        
+        if (Target.bBuildEditor == true)
+        {
+	        PublicDependencyModuleNames.AddRange(
+		        new string[]
+		        {
+			        "ApplicationCore",
+		        }
+	        );
+        }        
+    }
+}
