VisualGraphEdge: Add missing copyright boilerplate

diff --git a/Engine/Source/Developer/VisualGraphUtils/Private/VisualGraphEdge.cpp b/Engine/Source/Developer/VisualGraphUtils/Private/VisualGraphEdge.cpp
--- a/Engine/Source/Developer/VisualGraphUtils/Private/VisualGraphEdge.cpp
+++ b/Engine/Source/Developer/VisualGraphUtils/Private/VisualGraphEdge.cpp
@@ -1,4 +1,5 @@
-﻿
+﻿// Copyright Epic Games, Inc. All Rights Reserved.
+
 #include "VisualGraphEdge.h"
 #include "VisualGraph.h"
 
