VisualGraphTools: Fix missing nonunity include

diff --git a/Engine/Source/Developer/VisualGraphUtils/Private/VisualGraphObjectUtils.cpp b/Engine/Source/Developer/VisualGraphUtils/Private/VisualGraphObjectUtils.cpp
--- a/Engine/Source/Developer/VisualGraphUtils/Private/VisualGraphObjectUtils.cpp
+++ b/Engine/Source/Developer/VisualGraphUtils/Private/VisualGraphObjectUtils.cpp
@@ -6,6 +6,8 @@
 #include "Serialization/ArchiveUObject.h"
 #include "GameFramework/Actor.h"
 #include "Components/ActorComponent.h"
+#include "UObject/UObjectHash.h"
+#include "UObject/UObjectIterator.h"
 
 #if WITH_EDITOR
 
